Merging Pixel Streaming refactor using //UE5/Dev-TensorWorks-to-//UE5/Release-5.0
- Adds /?preferSFU option for browser peers to indicate they wish to use the SFU
- Adds an SFU server to Samples/PixelStreaming/WebServers/SFU
- Renames many classes in the PixelStreaming plugin to remove the "PixelStreaming" prefix from class names in private classes in favour of the namespace UE::PixelStreaming; public classes retain the prefix PixelStreaming and do not use the namespace.
- Renamed a number of classes to use the phrase "PixelStreaming" instead of "PixelStreamer", specifically, PixelStreamerInputComponent to PixelStreamingInputComponent
- All blueprint exposed static functions have been moved into a blueprint function library in PixelStreamingBlueprints.
- Fixed a memory leak caused by buffer allocation in the NVENC encoder that was never emptied.
- Moved to an explicit single threaded model for the NVENC video encoder - meaning only one encode job at a time per instance of the encoder.
- Moved to more explicit failure message when a user runs out of available NVENC sessions - this is common a consumer NVIDIA GPUs that only support 3 encoding sessions now that we are supporting simulcast.
- Introduced default simulcast settings of two streams: a full res stream and a half res stream. This can be overridden with `-SimulcastParameters=1.0,5000000,20000000,2.0,1000000,5000000` where each stream's parameters are specified in triples of scaling factor, min bitrate, max bitrate. Beware of running out of NVENC sessions if you specify too many stream sizes.
- Moved to more safety checking during shutdown of datachannel observer as this class was causing some intermittent failures to do a clean shutdown.
- Exposed WebRTC stats to Pixel Streaming. These can be viewed on screen using `-PixelStreamingHUDStats`.

diff --git a/Engine/Plugins/Media/HardwareEncoders/Source/EncoderNVENC/Private/NVENC_EncoderH264.cpp b/Engine/Plugins/Media/HardwareEncoders/Source/EncoderNVENC/Private/NVENC_EncoderH264.cpp
--- a/Engine/Plugins/Media/HardwareEncoders/Source/EncoderNVENC/Private/NVENC_EncoderH264.cpp
+++ b/Engine/Plugins/Media/HardwareEncoders/Source/EncoderNVENC/Private/NVENC_EncoderH264.cpp
@@ -25,55 +25,55 @@ namespace
 {
 	NV_ENC_PARAMS_RC_MODE ConvertRateControlModeNVENC(AVEncoder::FVideoEncoder::RateControlMode mode)
 	{
-		switch(mode)
+		switch (mode)
 		{
-		case AVEncoder::FVideoEncoder::RateControlMode::CONSTQP:
-			return NV_ENC_PARAMS_RC_CONSTQP;
-		case AVEncoder::FVideoEncoder::RateControlMode::VBR:
-			return NV_ENC_PARAMS_RC_VBR;
-		default:
-		case AVEncoder::FVideoEncoder::RateControlMode::CBR:
-			return NV_ENC_PARAMS_RC_CBR;
+			case AVEncoder::FVideoEncoder::RateControlMode::CONSTQP:
+				return NV_ENC_PARAMS_RC_CONSTQP;
+			case AVEncoder::FVideoEncoder::RateControlMode::VBR:
+				return NV_ENC_PARAMS_RC_VBR;
+			default:
+			case AVEncoder::FVideoEncoder::RateControlMode::CBR:
+				return NV_ENC_PARAMS_RC_CBR;
 		}
 	}
 
 	NV_ENC_MULTI_PASS ConvertMultipassModeNVENC(AVEncoder::FVideoEncoder::MultipassMode mode)
 	{
-		switch(mode)
+		switch (mode)
 		{
-		case AVEncoder::FVideoEncoder::MultipassMode::DISABLED:
-			return NV_ENC_MULTI_PASS_DISABLED;
-		case AVEncoder::FVideoEncoder::MultipassMode::QUARTER:
-			return NV_ENC_TWO_PASS_QUARTER_RESOLUTION;
-		default:
-		case AVEncoder::FVideoEncoder::MultipassMode::FULL:
-			return NV_ENC_TWO_PASS_FULL_RESOLUTION;
+			case AVEncoder::FVideoEncoder::MultipassMode::DISABLED:
+				return NV_ENC_MULTI_PASS_DISABLED;
+			case AVEncoder::FVideoEncoder::MultipassMode::QUARTER:
+				return NV_ENC_TWO_PASS_QUARTER_RESOLUTION;
+			default:
+			case AVEncoder::FVideoEncoder::MultipassMode::FULL:
+				return NV_ENC_TWO_PASS_FULL_RESOLUTION;
 		}
 	}
 
 	GUID ConvertH264Profile(AVEncoder::FVideoEncoder::H264Profile profile)
 	{
-		switch(profile)
-		{
-		default:
-		case AVEncoder::FVideoEncoder::H264Profile::AUTO:
-			return NV_ENC_CODEC_PROFILE_AUTOSELECT_GUID;
-		case AVEncoder::FVideoEncoder::H264Profile::BASELINE:
-			return NV_ENC_H264_PROFILE_BASELINE_GUID;
-		case AVEncoder::FVideoEncoder::H264Profile::MAIN:
-			return NV_ENC_H264_PROFILE_MAIN_GUID;
-		case AVEncoder::FVideoEncoder::H264Profile::HIGH:
-			return NV_ENC_H264_PROFILE_HIGH_GUID;
-		case AVEncoder::FVideoEncoder::H264Profile::HIGH444:
-			return NV_ENC_H264_PROFILE_HIGH_444_GUID;
-		case AVEncoder::FVideoEncoder::H264Profile::STEREO:
-			return NV_ENC_H264_PROFILE_STEREO_GUID;
-		case AVEncoder::FVideoEncoder::H264Profile::SVC_TEMPORAL_SCALABILITY:
-			return NV_ENC_H264_PROFILE_SVC_TEMPORAL_SCALABILTY;
-		case AVEncoder::FVideoEncoder::H264Profile::PROGRESSIVE_HIGH:
-			return NV_ENC_H264_PROFILE_PROGRESSIVE_HIGH_GUID;
-		case AVEncoder::FVideoEncoder::H264Profile::CONSTRAINED_HIGH:
-			return NV_ENC_H264_PROFILE_CONSTRAINED_HIGH_GUID;
+		switch (profile)
+		{
+			default:
+			case AVEncoder::FVideoEncoder::H264Profile::AUTO:
+				return NV_ENC_CODEC_PROFILE_AUTOSELECT_GUID;
+			case AVEncoder::FVideoEncoder::H264Profile::BASELINE:
+				return NV_ENC_H264_PROFILE_BASELINE_GUID;
+			case AVEncoder::FVideoEncoder::H264Profile::MAIN:
+				return NV_ENC_H264_PROFILE_MAIN_GUID;
+			case AVEncoder::FVideoEncoder::H264Profile::HIGH:
+				return NV_ENC_H264_PROFILE_HIGH_GUID;
+			case AVEncoder::FVideoEncoder::H264Profile::HIGH444:
+				return NV_ENC_H264_PROFILE_HIGH_444_GUID;
+			case AVEncoder::FVideoEncoder::H264Profile::STEREO:
+				return NV_ENC_H264_PROFILE_STEREO_GUID;
+			case AVEncoder::FVideoEncoder::H264Profile::SVC_TEMPORAL_SCALABILITY:
+				return NV_ENC_H264_PROFILE_SVC_TEMPORAL_SCALABILTY;
+			case AVEncoder::FVideoEncoder::H264Profile::PROGRESSIVE_HIGH:
+				return NV_ENC_H264_PROFILE_PROGRESSIVE_HIGH_GUID;
+			case AVEncoder::FVideoEncoder::H264Profile::CONSTRAINED_HIGH:
+				return NV_ENC_H264_PROFILE_CONSTRAINED_HIGH_GUID;
 		}
 	}
 } // namespace
@@ -82,37 +82,37 @@ namespace AVEncoder
 {
 	// Console variables for NVENC
 
-	TAutoConsoleVariable<int32>  CVarNVENCIntraRefreshPeriodFrames(
-	TEXT("NVENC.IntraRefreshPeriodFrames"),
-	0,
-	TEXT("The total number of frames between each intra refresh. Smallers values will cause intra refresh more often. Default: 0. Values <= 0 will disable intra refresh."),
-	ECVF_Default);
-
-	TAutoConsoleVariable<int32>  CVarNVENCIntraRefreshCountFrames(
-	TEXT("NVENC.IntraRefreshCountFrames"),
-	0,
-	TEXT("The total number of frames within the intra refresh period that should be used as 'intra refresh' frames. Smaller values make stream recovery quicker at the cost of more bandwidth usage. Default: 0."),
-	ECVF_Default);
-
-	TAutoConsoleVariable<bool>  CVarNVENCKeyframeQPUseLastQP(
-	TEXT("NVENC.KeyframeQPUseLastQP"),
-	true,
-	TEXT("If true QP of keyframes is no worse than the last frame transmitted (may cost latency), if false, it may be keyframe QP may be worse if network conditions require it (lower latency/worse quality). Default: true."),
-	ECVF_Default);
-
-	TAutoConsoleVariable<bool>  CVarNVENCEnableStats(
-	TEXT("NVENC.EnableStats"),
-	false,
-	TEXT("Whether to enable NVENC stats or not. Default: false."),
-	ECVF_Default);
-
-	TAutoConsoleVariable<int32>  CVarNVENCKeyframeInterval(
-	TEXT("NVENC.KeyframeInterval"),
-	300,
-	TEXT("Every N frames an IDR frame is sent. Default: 300. Note: A value <= 0 will disable sending of IDR frames on an interval."),
-	ECVF_Default);
-
-	template<typename T>
+	TAutoConsoleVariable<int32> CVarNVENCIntraRefreshPeriodFrames(
+		TEXT("NVENC.IntraRefreshPeriodFrames"),
+		0,
+		TEXT("The total number of frames between each intra refresh. Smallers values will cause intra refresh more often. Default: 0. Values <= 0 will disable intra refresh."),
+		ECVF_Default);
+
+	TAutoConsoleVariable<int32> CVarNVENCIntraRefreshCountFrames(
+		TEXT("NVENC.IntraRefreshCountFrames"),
+		0,
+		TEXT("The total number of frames within the intra refresh period that should be used as 'intra refresh' frames. Smaller values make stream recovery quicker at the cost of more bandwidth usage. Default: 0."),
+		ECVF_Default);
+
+	TAutoConsoleVariable<bool> CVarNVENCKeyframeQPUseLastQP(
+		TEXT("NVENC.KeyframeQPUseLastQP"),
+		true,
+		TEXT("If true QP of keyframes is no worse than the last frame transmitted (may cost latency), if false, it may be keyframe QP may be worse if network conditions require it (lower latency/worse quality). Default: true."),
+		ECVF_Default);
+
+	TAutoConsoleVariable<bool> CVarNVENCEnableStats(
+		TEXT("NVENC.EnableStats"),
+		false,
+		TEXT("Whether to enable NVENC stats or not. Default: false."),
+		ECVF_Default);
+
+	TAutoConsoleVariable<int32> CVarNVENCKeyframeInterval(
+		TEXT("NVENC.KeyframeInterval"),
+		300,
+		TEXT("Every N frames an IDR frame is sent. Default: 300. Note: A value <= 0 will disable sending of IDR frames on an interval."),
+		ECVF_Default);
+
+	template <typename T>
 	void NVENCCommandLineParseValue(const TCHAR* Match, TAutoConsoleVariable<T>& CVar)
 	{
 		T Value;
@@ -127,11 +127,14 @@ namespace AVEncoder
 		FString ValueMatch(Match);
 		ValueMatch.Append(TEXT("="));
 		FString Value;
-		if (FParse::Value(FCommandLine::Get(), *ValueMatch, Value)) {
-			if (Value.Equals(FString(TEXT("true")), ESearchCase::IgnoreCase)) {
+		if (FParse::Value(FCommandLine::Get(), *ValueMatch, Value))
+		{
+			if (Value.Equals(FString(TEXT("true")), ESearchCase::IgnoreCase))
+			{
 				CVar->Set(true, ECVF_SetByCommandline);
 			}
-			else if (Value.Equals(FString(TEXT("false")), ESearchCase::IgnoreCase)) {
+			else if (Value.Equals(FString(TEXT("false")), ESearchCase::IgnoreCase))
+			{
 				CVar->Set(false, ECVF_SetByCommandline);
 			}
 		}
@@ -145,8 +148,7 @@ namespace AVEncoder
 	{
 		// CVar changes can only be triggered from the game thread
 
-		AsyncTask(ENamedThreads::GameThread, []()
-		{ 
+		AsyncTask(ENamedThreads::GameThread, []() {
 			// Values
 			NVENCCommandLineParseValue(TEXT("-NVENCIntraRefreshPeriodFrames="), CVarNVENCIntraRefreshPeriodFrames);
 			NVENCCommandLineParseValue(TEXT("-NVENCIntraRefreshCountFrames="), CVarNVENCIntraRefreshCountFrames);
@@ -157,17 +159,15 @@ namespace AVEncoder
 			NVENCCommandLineParseOption(TEXT("-NVENCEnableStats"), CVarNVENCEnableStats);
 
 			// When NVENC stats CVar changes, change the output to screen flag.
-            CVarNVENCEnableStats.AsVariable()->SetOnChangedCallback(FConsoleVariableDelegate::CreateLambda([](IConsoleVariable* ChangedCVar)
-            {
-                FNVENCStats::Get().SetOutputToScreen(CVarNVENCEnableStats.GetValueOnAnyThread());
-            }));
+			CVarNVENCEnableStats.AsVariable()->SetOnChangedCallback(FConsoleVariableDelegate::CreateLambda([](IConsoleVariable* ChangedCVar) {
+				FNVENCStats::Get().SetOutputToScreen(CVarNVENCEnableStats.GetValueOnAnyThread());
+			}));
 
-			if(CVarNVENCEnableStats.GetValueOnAnyThread())
+			if (CVarNVENCEnableStats.GetValueOnAnyThread())
 			{
 				FNVENCStats::Get().SetOutputToScreen(true);
 			}
 		});
-        
 	}
 
 	static bool GetEncoderInfo(FNVENCCommon& NVENC, FVideoEncoderInfo& EncoderInfo);
@@ -177,7 +177,7 @@ namespace AVEncoder
 	{
 		FNVENCCommon& NVENC = FNVENCCommon::Setup();
 		bool bIsAvailable = NVENC.GetIsAvailable();
-		if(bIsAvailable)
+		if (bIsAvailable)
 		{
 			OutEncoderInfo.CodecType = ECodecType::H264;
 		}
@@ -187,17 +187,18 @@ namespace AVEncoder
 	void FVideoEncoderNVENC_H264::Register(FVideoEncoderFactory& InFactory)
 	{
 		FNVENCCommon& NVENC = FNVENCCommon::Setup();
-		if(NVENC.GetIsAvailable() && IsRHIDeviceNVIDIA())
+		if (NVENC.GetIsAvailable() && IsRHIDeviceNVIDIA())
 		{
 			FVideoEncoderInfo EncoderInfo;
-			if(GetEncoderInfo(NVENC, EncoderInfo))
+			if (GetEncoderInfo(NVENC, EncoderInfo))
 			{
 				InFactory.Register(EncoderInfo, []() { return TUniquePtr<FVideoEncoder>(new FVideoEncoderNVENC_H264()); });
 			}
 		}
 	}
 
-	FVideoEncoderNVENC_H264::FVideoEncoderNVENC_H264() : NVENC(FNVENCCommon::Setup()) 
+	FVideoEncoderNVENC_H264::FVideoEncoderNVENC_H264()
+		: NVENC(FNVENCCommon::Setup())
 	{
 		// Parse NVENC settings from command line (if any relevant ones are passed)
 		NVENCParseCommandLineFlags();
@@ -207,38 +208,38 @@ namespace AVEncoder
 
 	bool FVideoEncoderNVENC_H264::Setup(TSharedRef<FVideoEncoderInput> InputFrameFactory, const FLayerConfig& InitConfig)
 	{
-		if(!NVENC.GetIsAvailable())
+		if (!NVENC.GetIsAvailable())
 		{
 			UE_LOG(LogEncoderNVENC, Error, TEXT("NVENC not avaliable"));
 			return false;
 		}
 
 		FrameFormat = InputFrameFactory->GetFrameFormat();
-		switch(FrameFormat)
+		switch (FrameFormat)
 		{
 #if PLATFORM_WINDOWS
-		case AVEncoder::EVideoFrameFormat::D3D11_R8G8B8A8_UNORM:
-		case AVEncoder::EVideoFrameFormat::D3D12_R8G8B8A8_UNORM:
-			EncoderDevice = InputFrameFactory->GetD3D11EncoderDevice();
-			break;
+			case AVEncoder::EVideoFrameFormat::D3D11_R8G8B8A8_UNORM:
+			case AVEncoder::EVideoFrameFormat::D3D12_R8G8B8A8_UNORM:
+				EncoderDevice = InputFrameFactory->GetD3D11EncoderDevice();
+				break;
 #endif
-		case AVEncoder::EVideoFrameFormat::CUDA_R8G8B8A8_UNORM:
-			EncoderDevice = InputFrameFactory->GetCUDAEncoderContext();
-			break;
-		case AVEncoder::EVideoFrameFormat::Undefined:
-		default:
-			UE_LOG(LogEncoderNVENC, Error, TEXT("Frame format %s is not supported by NVENC_Encoder on this platform."), *ToString(FrameFormat));
-			return false;
+			case AVEncoder::EVideoFrameFormat::CUDA_R8G8B8A8_UNORM:
+				EncoderDevice = InputFrameFactory->GetCUDAEncoderContext();
+				break;
+			case AVEncoder::EVideoFrameFormat::Undefined:
+			default:
+				UE_LOG(LogEncoderNVENC, Error, TEXT("Frame format %s is not supported by NVENC_Encoder on this platform."), *ToString(FrameFormat));
+				return false;
 		}
 
-		if(!EncoderDevice)
+		if (!EncoderDevice)
 		{
 			UE_LOG(LogEncoderNVENC, Error, TEXT("NVENC needs encoder device."));
 			return false;
 		}
 
 		FLayerConfig mutableConfig = InitConfig;
-		if(mutableConfig.MaxFramerate == 0)
+		if (mutableConfig.MaxFramerate == 0)
 		{
 			mutableConfig.MaxFramerate = 60;
 		}
@@ -249,7 +250,7 @@ namespace AVEncoder
 	FVideoEncoder::FLayer* FVideoEncoderNVENC_H264::CreateLayer(uint32 layerIdx, FLayerConfig const& config)
 	{
 		auto const layer = new FNVENCLayer(layerIdx, config, *this);
-		if(!layer->Setup())
+		if (!layer->Setup())
 		{
 			delete layer;
 			return nullptr;
@@ -261,7 +262,7 @@ namespace AVEncoder
 
 	void FVideoEncoderNVENC_H264::Encode(FVideoEncoderInputFrame const* InFrame, const FEncodeOptions& EncodeOptions)
 	{
-		for(FVideoEncoder::FLayer* Layer : Layers)
+		for (FVideoEncoder::FLayer* Layer : Layers)
 		{
 			FVideoEncoderNVENC_H264::FNVENCLayer* NVLayer = static_cast<FVideoEncoderNVENC_H264::FNVENCLayer*>(Layer);
 			NVLayer->Encode(InFrame, EncodeOptions);
@@ -270,7 +271,7 @@ namespace AVEncoder
 
 	void FVideoEncoderNVENC_H264::Flush()
 	{
-		for(FVideoEncoder::FLayer* Layer : Layers)
+		for (FVideoEncoder::FLayer* Layer : Layers)
 		{
 			FVideoEncoderNVENC_H264::FNVENCLayer* NVLayer = static_cast<FVideoEncoderNVENC_H264::FNVENCLayer*>(Layer);
 			NVLayer->Flush();
@@ -279,7 +280,7 @@ namespace AVEncoder
 
 	void FVideoEncoderNVENC_H264::Shutdown()
 	{
-		for(FVideoEncoder::FLayer* Layer : Layers)
+		for (FVideoEncoder::FLayer* Layer : Layers)
 		{
 			FVideoEncoderNVENC_H264::FNVENCLayer* NVLayer = static_cast<FVideoEncoderNVENC_H264::FNVENCLayer*>(Layer);
 			NVLayer->Shutdown();
@@ -290,7 +291,11 @@ namespace AVEncoder
 
 	// --- FVideoEncoderNVENC_H264::FLayer ------------------------------------------------------------
 	FVideoEncoderNVENC_H264::FNVENCLayer::FNVENCLayer(uint32 layerIdx, FLayerConfig const& config, FVideoEncoderNVENC_H264& encoder)
-		: FLayer(config), Encoder(encoder), NVENC(FNVENCCommon::Setup()), CodecGUID(NV_ENC_CODEC_H264_GUID), LayerIndex(layerIdx)
+		: FLayer(config)
+		, Encoder(encoder)
+		, NVENC(FNVENCCommon::Setup())
+		, CodecGUID(NV_ENC_CODEC_H264_GUID)
+		, LayerIndex(layerIdx)
 	{
 	}
 
@@ -298,11 +303,11 @@ namespace AVEncoder
 
 	bool FVideoEncoderNVENC_H264::FNVENCLayer::Setup()
 	{
-		if(CreateSession() && CreateInitialConfig())
+		if (CreateSession() && CreateInitialConfig())
 		{
 			// create encoder
 			NVENCSTATUS Result = NVENC.nvEncInitializeEncoder(NVEncoder, &EncoderInitParams);
-			if(Result != NV_ENC_SUCCESS)
+			if (Result != NV_ENC_SUCCESS)
 			{
 				UE_LOG(LogEncoderNVENC, Error, TEXT("Unable to initialize NvEnc encoder (%s)."), *NVENC.GetErrorString(NVEncoder, Result));
 				return false;
@@ -318,30 +323,30 @@ namespace AVEncoder
 
 	bool FVideoEncoderNVENC_H264::FNVENCLayer::CreateSession()
 	{
-		if(!NVEncoder)
+		if (!NVEncoder)
 		{
 			// create the encoder session
 			NVENCStruct(NV_ENC_OPEN_ENCODE_SESSION_EX_PARAMS, OpenEncodeSessionExParams);
 			OpenEncodeSessionExParams.apiVersion = NVENCAPI_VERSION;
 			OpenEncodeSessionExParams.device = Encoder.EncoderDevice;
 
-			switch(Encoder.FrameFormat)
+			switch (Encoder.FrameFormat)
 			{
-			case EVideoFrameFormat::D3D11_R8G8B8A8_UNORM:
-			case EVideoFrameFormat::D3D12_R8G8B8A8_UNORM:
-				OpenEncodeSessionExParams.deviceType = NV_ENC_DEVICE_TYPE_DIRECTX;
-				break;
-			case EVideoFrameFormat::CUDA_R8G8B8A8_UNORM:
-				OpenEncodeSessionExParams.deviceType = NV_ENC_DEVICE_TYPE_CUDA;
-				break;
-			default:
-				UE_LOG(LogEncoderNVENC, Error, TEXT("FrameFormat %s unavailable."), *ToString(Encoder.FrameFormat));
-				return false;
-				break;
+				case EVideoFrameFormat::D3D11_R8G8B8A8_UNORM:
+				case EVideoFrameFormat::D3D12_R8G8B8A8_UNORM:
+					OpenEncodeSessionExParams.deviceType = NV_ENC_DEVICE_TYPE_DIRECTX;
+					break;
+				case EVideoFrameFormat::CUDA_R8G8B8A8_UNORM:
+					OpenEncodeSessionExParams.deviceType = NV_ENC_DEVICE_TYPE_CUDA;
+					break;
+				default:
+					UE_LOG(LogEncoderNVENC, Error, TEXT("FrameFormat %s unavailable."), *ToString(Encoder.FrameFormat));
+					return false;
+					break;
 			}
 
 			auto const result = NVENC.nvEncOpenEncodeSessionEx(&OpenEncodeSessionExParams, &NVEncoder);
-			if(result != NV_ENC_SUCCESS)
+			if (result != NV_ENC_SUCCESS)
 			{
 				UE_LOG(LogEncoderNVENC, Error, TEXT("Unable to open NvEnc encoding session (%s)."), *NVENC.GetErrorString(NVEncoder, result));
 				NVEncoder = nullptr;
@@ -375,7 +380,7 @@ namespace AVEncoder
 		NVENCStruct(NV_ENC_PRESET_CONFIG, PresetConfig);
 		PresetConfig.presetCfg.version = NV_ENC_CONFIG_VER;
 		auto const result = NVENC.nvEncGetEncodePresetConfigEx(NVEncoder, EncoderInitParams.encodeGUID, EncoderInitParams.presetGUID, EncoderInitParams.tuningInfo, &PresetConfig);
-		if(result != NV_ENC_SUCCESS)
+		if (result != NV_ENC_SUCCESS)
 		{
 			UE_LOG(LogEncoderNVENC, Error, TEXT("Unable to get NvEnc preset config (%s)."), *NVENC.GetErrorString(NVEncoder, result));
 			return false;
@@ -397,20 +402,20 @@ namespace AVEncoder
 		*/
 		int32 IntraRefreshPeriodFrames = CVarNVENCIntraRefreshPeriodFrames.GetValueOnAnyThread();
 		int32 IntraRefreshCountFrames = CVarNVENCIntraRefreshCountFrames.GetValueOnAnyThread();
-		bool bIntraRefreshSupported = GetEncoderCapability(NVENC, NVEncoder, NV_ENC_CAPS_SUPPORT_INTRA_REFRESH) > 0; 
+		bool bIntraRefreshSupported = GetEncoderCapability(NVENC, NVEncoder, NV_ENC_CAPS_SUPPORT_INTRA_REFRESH) > 0;
 		bool bIntraRefreshEnabled = IntraRefreshPeriodFrames > 0;
 
-		if (bIntraRefreshEnabled && bIntraRefreshSupported) 
+		if (bIntraRefreshEnabled && bIntraRefreshSupported)
 		{
 			EncoderConfig.encodeCodecConfig.h264Config.enableIntraRefresh = 1;
 			EncoderConfig.encodeCodecConfig.h264Config.intraRefreshPeriod = IntraRefreshPeriodFrames;
 			EncoderConfig.encodeCodecConfig.h264Config.intraRefreshCnt = IntraRefreshCountFrames;
-			
+
 			UE_LOG(LogEncoderNVENC, Log, TEXT("NVENC intra refresh enabled."));
 			UE_LOG(LogEncoderNVENC, Log, TEXT("NVENC intra refresh period set to = %d"), IntraRefreshPeriodFrames);
 			UE_LOG(LogEncoderNVENC, Log, TEXT("NVENC intra refresh count = %d"), IntraRefreshCountFrames);
 		}
-		else if(bIntraRefreshEnabled && !bIntraRefreshSupported)
+		else if (bIntraRefreshEnabled && !bIntraRefreshSupported)
 		{
 			UE_LOG(LogEncoderNVENC, Error, TEXT("NVENC intra refresh capability is not supported on this device, cannot use this feature"));
 		}
@@ -419,7 +424,7 @@ namespace AVEncoder
 		* IDR period - how often to send IDR (instantaneous decode refresh) frames, a.k.a keyframes. This can stabilise a stream that dropped/lost some frames (but at the cost of more bandwidth).
 		*/
 		int32 IdrPeriod = CVarNVENCKeyframeInterval.GetValueOnAnyThread();
-		if(IdrPeriod > 0)
+		if (IdrPeriod > 0)
 		{
 			EncoderConfig.encodeCodecConfig.h264Config.idrPeriod = IdrPeriod;
 		}
@@ -451,11 +456,11 @@ namespace AVEncoder
 		uint64 PreReconfigureCycles = FPlatformTime::Cycles64();
 
 		FScopeLock lock(&ConfigMutex);
-		if(NeedsReconfigure)
+		if (NeedsReconfigure)
 		{
 			UpdateConfig();
-			
-			if(EncoderInitParams.frameRateNum != CurrentConfig.MaxFramerate)
+
+			if (EncoderInitParams.frameRateNum != CurrentConfig.MaxFramerate)
 			{
 				EncoderInitParams.frameRateNum = CurrentConfig.MaxFramerate;
 			}
@@ -464,14 +469,14 @@ namespace AVEncoder
 			FMemory::Memcpy(&ReconfigureParams.reInitEncodeParams, &EncoderInitParams, sizeof(EncoderInitParams));
 
 			auto const result = NVENC.nvEncReconfigureEncoder(NVEncoder, &ReconfigureParams);
-			if(result != NV_ENC_SUCCESS)
+			if (result != NV_ENC_SUCCESS)
 				UE_LOG(LogEncoderNVENC, Error, TEXT("Failed to update NVENC encoder configuration (%s)"), *NVENC.GetErrorString(NVEncoder, result));
 
 			NeedsReconfigure = false;
 
 			uint64 PostReconfigureCycles = FPlatformTime::Cycles64();
 
-			if(CVarNVENCEnableStats.GetValueOnAnyThread())
+			if (CVarNVENCEnableStats.GetValueOnAnyThread())
 			{
 				// The time it took to do nvEncEncodePicture
 				double ReconfigureDeltaMs = FPlatformTime::ToMilliseconds64(PostReconfigureCycles - PreReconfigureCycles);
@@ -482,7 +487,7 @@ namespace AVEncoder
 
 	void FVideoEncoderNVENC_H264::FNVENCLayer::UpdateLastEncodedQP(uint32 InLastEncodedQP)
 	{
-		if(InLastEncodedQP == LastEncodedQP)
+		if (InLastEncodedQP == LastEncodedQP)
 		{
 			// QP is the same, do nothing.
 			return;
@@ -505,14 +510,14 @@ namespace AVEncoder
 		RateControlParams.averageBitRate = CurrentConfig.TargetBitrate > -1 ? CurrentConfig.TargetBitrate : DEFAULT_BITRATE;
 		RateControlParams.maxBitRate = CurrentConfig.MaxBitrate > -1 ? CurrentConfig.MaxBitrate : DEFAULT_BITRATE; // Not used for CBR
 		RateControlParams.multiPass = ConvertMultipassModeNVENC(CurrentConfig.MultipassMode);
-		RateControlParams.minQP = {MinQP, MinQP, MinQP};
-		RateControlParams.maxQP = {MaxQP, MaxQP, MaxQP}; 
+		RateControlParams.minQP = { MinQP, MinQP, MinQP };
+		RateControlParams.maxQP = { MaxQP, MaxQP, MaxQP };
 		RateControlParams.enableMinQP = CurrentConfig.QPMin > -1;
 		RateControlParams.enableMaxQP = CurrentConfig.QPMax > -1;
 
 		// If we have QP ranges turned on use the last encoded QP to guide the max QP for an i-frame, so the i-frame doesn't look too blocky
 		// Note: this does nothing if we have i-frames turned off.
-		if(RateControlParams.enableMaxQP && LastEncodedQP > 0 && CVarNVENCKeyframeQPUseLastQP.GetValueOnAnyThread())
+		if (RateControlParams.enableMaxQP && LastEncodedQP > 0 && CVarNVENCKeyframeQPUseLastQP.GetValueOnAnyThread())
 		{
 			RateControlParams.maxQP.qpIntra = LastEncodedQP;
 		}
@@ -522,7 +527,7 @@ namespace AVEncoder
 		NV_ENC_CONFIG_H264& H264Config = EncoderInitParams.encodeConfig->encodeCodecConfig.h264Config;
 		H264Config.enableFillerDataInsertion = CurrentConfig.FillData ? 1 : 0;
 
-		if(CurrentConfig.RateControlMode == AVEncoder::FVideoEncoder::RateControlMode::CBR && CurrentConfig.FillData)
+		if (CurrentConfig.RateControlMode == AVEncoder::FVideoEncoder::RateControlMode::CBR && CurrentConfig.FillData)
 		{
 			// `outputPictureTimingSEI` is used in CBR mode to fill video frame with data to match the requested bitrate.
 			H264Config.outputPictureTimingSEI = 1;
@@ -532,7 +537,7 @@ namespace AVEncoder
 	void FVideoEncoderNVENC_H264::FNVENCLayer::EncodeBuffer(FInputOutput* Buffer)
 	{
 		checkf(Buffer, TEXT("Cannot encode null buffer."));
-		if(Buffer == nullptr)
+		if (Buffer == nullptr)
 		{
 			UE_LOG(LogEncoderNVENC, Error, TEXT("Cannot proceed with encoding in NVENC - buffer was nullptr."));
 			return;
@@ -547,13 +552,13 @@ namespace AVEncoder
 
 		uint64 StartQueueEncodeCycles = FPlatformTime::Cycles64();
 
-		if(Buffer && !RegisterInputTexture(Buffer))
+		if (Buffer && !RegisterInputTexture(Buffer))
 		{
 			UE_LOG(LogEncoderNVENC, Error, TEXT("Unable to register texture in NVENC"));
 			DestroyBuffer(Buffer);
 			return;
 		}
-		
+
 		Buffer->SubmitTimeCycles = StartQueueEncodeCycles;
 		Buffer->PicParams.version = NV_ENC_PIC_PARAMS_VER;
 		Buffer->PicParams.inputWidth = Buffer->Width;
@@ -567,7 +572,7 @@ namespace AVEncoder
 		MaybeReconfigure();
 
 		// check if we have not reconfigued the buffer size
-		if(Buffer->SourceFrame && (Buffer->Width != EncoderInitParams.encodeWidth || Buffer->Height != EncoderInitParams.encodeHeight))
+		if (Buffer->SourceFrame && (Buffer->Width != EncoderInitParams.encodeWidth || Buffer->Height != EncoderInitParams.encodeHeight))
 		{
 			DestroyBuffer(Buffer);
 			return;
@@ -578,7 +583,7 @@ namespace AVEncoder
 		NVENCSTATUS Result = NVENC.nvEncEncodePicture(NVEncoder, &(Buffer->PicParams));
 
 		uint64 PostnvEncEncodePictureCycles = FPlatformTime::Cycles64();
-		if(CVarNVENCEnableStats.GetValueOnAnyThread())
+		if (CVarNVENCEnableStats.GetValueOnAnyThread())
 		{
 			// The time it took to do nvEncEncodePicture
 			double ProcessFramesFuncDeltaMs = FPlatformTime::ToMilliseconds64(PostnvEncEncodePictureCycles - PrenvEncEncodePictureCycles);
@@ -586,13 +591,13 @@ namespace AVEncoder
 		}
 
 		// We have an encoded bitstream - to use it we must bitstream lock, read, unlock (and then pass out the encoded output the whoever needs it)
-		if(Result == NV_ENC_SUCCESS)
+		if (Result == NV_ENC_SUCCESS)
 		{
 			ProcessEncodedBuffer(Buffer);
 
 			UnregisterInputTexture(Buffer);
 
-			if(CVarNVENCEnableStats.GetValueOnAnyThread())
+			if (CVarNVENCEnableStats.GetValueOnAnyThread())
 			{
 				uint64 EndQueueEncodeCycles = FPlatformTime::Cycles64();
 				double QueueEncodeMs = FPlatformTime::ToMilliseconds64(EndQueueEncodeCycles - StartQueueEncodeCycles);
@@ -605,8 +610,8 @@ namespace AVEncoder
 		{
 			// Something went wrong
 			UE_LOG(LogEncoderNVENC, Error, TEXT("nvEncEncodePicture returned %s"), *NVENC.GetErrorString(NVEncoder, Result));
-			
-			if(Buffer)
+
+			if (Buffer)
 			{
 				DestroyBuffer(Buffer);
 			}
@@ -615,9 +620,21 @@ namespace AVEncoder
 
 	void FVideoEncoderNVENC_H264::FNVENCLayer::Encode(FVideoEncoderInputFrame const* InFrame, const FEncodeOptions& EncodeOptions)
 	{
+		// This encode is single threaded and blocking, if someone else is encoding right now, we wait.
+		while (bIsProcessingFrame) 
+		{
+			// According to UE docs
+			// Sleep this thread for Seconds. 0.0 means release the current time slice to let other threads get some attention.
+			FPlatformProcess::Sleep(0.0f);
+		}
+
+		checkf(!bIsProcessingFrame, TEXT("NVENC encoder should only ever be called one thread at a time."));
+
+		bIsProcessingFrame = true;
+
 		FInputOutput* Buffer = GetOrCreateBuffer(static_cast<const FVideoEncoderInputFrameImpl*>(InFrame));
 
-		if(!Buffer)
+		if (!Buffer)
 		{
 			UE_LOG(LogEncoderNVENC, Error, TEXT("Cannot encode null buffer."));
 			return;
@@ -627,7 +644,7 @@ namespace AVEncoder
 		Buffer->PicParams = {};
 		Buffer->PicParams.encodePicFlags = 0;
 
-		if(EncodeOptions.bForceKeyFrame)
+		if (EncodeOptions.bForceKeyFrame)
 		{
 			LastKeyFrameTime = FDateTime::UtcNow();
 			Buffer->PicParams.encodePicFlags |= NV_ENC_PIC_FLAG_FORCEIDR;
@@ -640,23 +657,13 @@ namespace AVEncoder
 		Buffer->PicParams.inputTimeStamp = Buffer->TimeStamp = InFrame->GetTimestampUs();
 
 		EncodeBuffer(Buffer);
+
+		bIsProcessingFrame = false;
 	}
 
 	void FVideoEncoderNVENC_H264::FNVENCLayer::Flush()
 	{
-		// Empty the FrameToBufferMapping
-
-		// As destroy buffer also modifies the frames to buffer mapping we must remove entries by copying keys first.
-		TArray<const AVEncoder::FVideoEncoderInputFrame*> Frames;
-		FrameToBufferMapping.GetKeys(Frames);
-		for(const AVEncoder::FVideoEncoderInputFrame* StoredFrame : Frames)
-		{
-			FInputOutput* Buffer = *FrameToBufferMapping.Find(StoredFrame);
-			DestroyBuffer(Buffer);
-		}
-
-		Frames.Empty();
-		FrameToBufferMapping.Empty();
+		DestroyBuffer(InputOutputBuffer);
 	}
 
 	void FVideoEncoderNVENC_H264::FNVENCLayer::Shutdown()
@@ -667,7 +674,7 @@ namespace AVEncoder
 		{
 			FInputOutput* EmptyBuffer = CreateBuffer();
 
-			if(!EmptyBuffer)
+			if (!EmptyBuffer)
 			{
 				return;
 			}
@@ -679,12 +686,12 @@ namespace AVEncoder
 			EncodeBuffer(EmptyBuffer);
 		}
 
-		if(NVEncoder)
+		if (NVEncoder)
 		{
 			UE_LOG(LogEncoderNVENC, Log, TEXT("Attempting to destroy NVENC encoder."));
 
 			auto const result = NVENC.nvEncDestroyEncoder(NVEncoder);
-			if(result != NV_ENC_SUCCESS)
+			if (result != NV_ENC_SUCCESS)
 			{
 				UE_LOG(LogEncoderNVENC, Error, TEXT("Failed to destroy NVENC encoder (%s)."), *NVENC.GetErrorString(NVEncoder, result));
 			}
@@ -696,7 +703,6 @@ namespace AVEncoder
 		}
 	}
 
-
 	void FVideoEncoderNVENC_H264::FNVENCLayer::ProcessEncodedBuffer(FInputOutput* Buffer)
 	{
 		uint64 StartProcessFramesFuncCycles;
@@ -705,13 +711,13 @@ namespace AVEncoder
 		StartProcessFramesFuncCycles = FPlatformTime::Cycles64();
 
 		// lock output buffers for CPU access
-		if(LockOutputBuffer(Buffer))
+		if (LockOutputBuffer(Buffer))
 		{
-			if(Encoder.OnEncodedPacket)
+			if (Encoder.OnEncodedPacket)
 			{
 				FCodecPacket Packet = FCodecPacket::Create(static_cast<const uint8*>(Buffer->BitstreamData), Buffer->BitstreamDataSize);
 
-				if(Buffer->PictureType & NV_ENC_PIC_TYPE_IDR)
+				if (Buffer->PictureType & NV_ENC_PIC_TYPE_IDR)
 				{
 					UE_LOG(LogEncoderNVENC, Verbose, TEXT("Generated IDR Frame"));
 					Packet.IsKeyFrame = true;
@@ -729,7 +735,7 @@ namespace AVEncoder
 				Packet.Framerate = EncoderInitParams.frameRateNum;
 
 				// Stats
-				if(CVarNVENCEnableStats.GetValueOnAnyThread())
+				if (CVarNVENCEnableStats.GetValueOnAnyThread())
 				{
 					// The time it took to do just the encode part
 					double ProcessFramesFuncDeltaMs = FPlatformTime::ToMilliseconds64(FinishedEncodingCycles - StartProcessFramesFuncCycles);
@@ -741,18 +747,17 @@ namespace AVEncoder
 				}
 
 				// Send frame to OnEncodedPacket listener, such as WebRTC.
-				if(Encoder.OnEncodedPacket)
+				if (Encoder.OnEncodedPacket)
 				{
 					Encoder.OnEncodedPacket(LayerIndex, Buffer->SourceFrame, Packet);
 				}
 			}
 
-			if(!UnlockOutputBuffer(Buffer))
+			if (!UnlockOutputBuffer(Buffer))
 			{
 				UE_LOG(LogEncoderNVENC, Error, TEXT("Failed to unlock output buffer."));
 				DestroyBuffer(Buffer);
 			}
-
 		}
 		else
 		{
@@ -767,7 +772,7 @@ namespace AVEncoder
 		NVENCStruct(NV_ENC_CAPS_PARAM, CapsParam);
 		CapsParam.capsToQuery = CapsToQuery;
 		auto const result = NVENC.nvEncGetEncodeCaps(NVEncoder, CodecGUID, &CapsParam, &CapsValue);
-		if(result != NV_ENC_SUCCESS)
+		if (result != NV_ENC_SUCCESS)
 		{
 			UE_LOG(LogEncoderNVENC, Warning, TEXT("Failed to query for NVENC capability %d (%s)."), CapsToQuery, *NVENC.GetErrorString(NVEncoder, result));
 			return 0;
@@ -777,38 +782,23 @@ namespace AVEncoder
 
 	FVideoEncoderNVENC_H264::FNVENCLayer::FInputOutput* FVideoEncoderNVENC_H264::FNVENCLayer::GetOrCreateBuffer(const FVideoEncoderInputFrameImpl* InFrame)
 	{
-		FInputOutput* Buffer = nullptr;
 
-		if(FrameToBufferMapping.Contains(InFrame))
+		if (InputOutputBuffer == nullptr)
 		{
-			Buffer = *FrameToBufferMapping.Find(InFrame);
-			checkf(Buffer, TEXT("Frame was mapped to a nullptr buffer, this indicate a buffer was not removed when it should have been."));
-
-			if(Buffer)
-			{
-				// Check for buffer and Frame texture resolution mismatch
-				if(InFrame->GetWidth() != Buffer->Width || InFrame->GetHeight() != Buffer->Height)
-				{
-					DestroyBuffer(Buffer);
-					Buffer = nullptr;
-				}
-				// Resolution of existing buffer is a match, so set our frame as the source frame
-				else
-				{
-					Buffer->SourceFrame = InFrame;
-				}
-			}
-			
+			InputOutputBuffer = CreateBuffer();
 		}
-		// No existing buffer for this frame
 		else
 		{
-			Buffer = CreateBuffer();
-			Buffer->SourceFrame = InFrame;
-			FrameToBufferMapping.Add(InFrame, Buffer);
+			// Check for buffer and Frame texture resolution mismatch
+			if (InFrame->GetWidth() != InputOutputBuffer->Width || InFrame->GetHeight() != InputOutputBuffer->Height)
+			{
+				DestroyBuffer(InputOutputBuffer);
+				InputOutputBuffer = CreateBuffer();
+			}
 		}
 
-		return Buffer;
+		InputOutputBuffer->SourceFrame = InFrame;
+		return InputOutputBuffer;
 	}
 
 	FVideoEncoderNVENC_H264::FNVENCLayer::FInputOutput* FVideoEncoderNVENC_H264::FNVENCLayer::CreateBuffer()
@@ -819,7 +809,7 @@ namespace AVEncoder
 		NVENCStruct(NV_ENC_CREATE_BITSTREAM_BUFFER, CreateParam);
 		{
 			auto const result = NVENC.nvEncCreateBitstreamBuffer(NVEncoder, &CreateParam);
-			if(result != NV_ENC_SUCCESS)
+			if (result != NV_ENC_SUCCESS)
 			{
 				UE_LOG(LogEncoderNVENC, Error, TEXT("Failed to create NVENC output buffer (%s)."), *NVENC.GetErrorString(NVEncoder, result));
 				DestroyBuffer(Buffer);
@@ -835,15 +825,14 @@ namespace AVEncoder
 	void FVideoEncoderNVENC_H264::FNVENCLayer::DestroyBuffer(FInputOutput* InBuffer)
 	{
 		checkf(InBuffer, TEXT("Cannot destroy buffer in NVENC - buffer was nullptr."));
-		if(!InBuffer)
+		if (!InBuffer)
 		{
 			return;
 		}
 
-		// Remove from our mapping of frames to buffers
-		if(InBuffer->SourceFrame)
+		// Release the AVEncoder frame
+		if (InBuffer->SourceFrame)
 		{
-			FrameToBufferMapping.Remove(InBuffer->SourceFrame);
 			InBuffer->SourceFrame->Release();
 		}
 
@@ -852,10 +841,10 @@ namespace AVEncoder
 
 		// Destroy output buffer - if any
 		UnlockOutputBuffer(InBuffer);
-		if(InBuffer->OutputBitstream)
+		if (InBuffer->OutputBitstream)
 		{
 			auto const result = NVENC.nvEncDestroyBitstreamBuffer(NVEncoder, InBuffer->OutputBitstream);
-			if(result != NV_ENC_SUCCESS)
+			if (result != NV_ENC_SUCCESS)
 			{
 				UE_LOG(LogEncoderNVENC, Warning, TEXT("Failed to destroy NVENC output buffer (%s)."), *NVENC.GetErrorString(NVEncoder, result));
 			}
@@ -866,12 +855,13 @@ namespace AVEncoder
 		InBuffer->InputTexture = nullptr;
 
 		delete InBuffer;
+		InBuffer = nullptr;
 	}
 
 	void FVideoEncoderNVENC_H264::FNVENCLayer::CreateResourceDIRECTX(FInputOutput* InBuffer, NV_ENC_REGISTER_RESOURCE& RegisterParam, FIntPoint TextureSize)
 	{
 		checkf(InBuffer, TEXT("Cannot create DirectX resource in NVENC - buffer was nullptr."));
-		if(!InBuffer)
+		if (!InBuffer)
 		{
 			return;
 		}
@@ -880,20 +870,20 @@ namespace AVEncoder
 		D3D11_TEXTURE2D_DESC Desc;
 		static_cast<ID3D11Texture2D*>(InBuffer->InputTexture)->GetDesc(&Desc);
 
-		switch(Desc.Format)
-		{
-		case DXGI_FORMAT_NV12:
-			InBuffer->BufferFormat = NV_ENC_BUFFER_FORMAT_NV12;
-			break;
-		case DXGI_FORMAT_R8G8B8A8_UNORM:
-			InBuffer->BufferFormat = NV_ENC_BUFFER_FORMAT_ABGR;
-			break;
-		case DXGI_FORMAT_B8G8R8A8_UNORM:
-			InBuffer->BufferFormat = NV_ENC_BUFFER_FORMAT_ARGB;
-			break;
-		default:
-			UE_LOG(LogEncoderNVENC, Error, TEXT("Invalid input texture format for NVENC (%d)"), Desc.Format);
-			return;
+		switch (Desc.Format)
+		{
+			case DXGI_FORMAT_NV12:
+				InBuffer->BufferFormat = NV_ENC_BUFFER_FORMAT_NV12;
+				break;
+			case DXGI_FORMAT_R8G8B8A8_UNORM:
+				InBuffer->BufferFormat = NV_ENC_BUFFER_FORMAT_ABGR;
+				break;
+			case DXGI_FORMAT_B8G8R8A8_UNORM:
+				InBuffer->BufferFormat = NV_ENC_BUFFER_FORMAT_ARGB;
+				break;
+			default:
+				UE_LOG(LogEncoderNVENC, Error, TEXT("Invalid input texture format for NVENC (%d)"), Desc.Format);
+				return;
 		}
 
 		InBuffer->Width = TextureSize.X;
@@ -912,12 +902,13 @@ namespace AVEncoder
 	void FVideoEncoderNVENC_H264::FNVENCLayer::CreateResourceCUDAARRAY(FInputOutput* InBuffer, NV_ENC_REGISTER_RESOURCE& RegisterParam, FIntPoint TextureSize)
 	{
 		checkf(InBuffer, TEXT("Cannot create CUDA array in NVENC - buffer was nullptr."));
-		if(!InBuffer)
+		if (!InBuffer)
 		{
 			return;
 		}
 
-		switch(InBuffer->SourceFrame->GetCUDA().UnderlyingRHI) {
+		switch (InBuffer->SourceFrame->GetCUDA().UnderlyingRHI)
+		{
 			case FVideoEncoderInputFrame::EUnderlyingRHI::Vulkan:
 				InBuffer->Pitch = TextureSize.X * GPixelFormats[EPixelFormat::PF_A8R8G8B8].BlockBytes;
 				break;
@@ -944,36 +935,36 @@ namespace AVEncoder
 	bool FVideoEncoderNVENC_H264::FNVENCLayer::RegisterInputTexture(FInputOutput* InBuffer)
 	{
 		checkf(InBuffer, TEXT("Cannot register texture in NVENC - buffer was nullptr."));
-		if(!InBuffer)
+		if (!InBuffer)
 		{
 			return false;
 		}
 
-		if(!InBuffer->InputTexture && InBuffer->PicParams.encodePicFlags != NV_ENC_PIC_FLAG_EOS)
+		if (!InBuffer->InputTexture && InBuffer->PicParams.encodePicFlags != NV_ENC_PIC_FLAG_EOS)
 		{
 			checkf(InBuffer->SourceFrame, TEXT("Cannot register texture if source frame is null."));
 
 			void* TextureToCompress = nullptr;
 
-			switch(InBuffer->SourceFrame->GetFormat())
+			switch (InBuffer->SourceFrame->GetFormat())
 			{
 #if PLATFORM_WINDOWS
-			case EVideoFrameFormat::D3D11_R8G8B8A8_UNORM:
-				TextureToCompress = InBuffer->SourceFrame->GetD3D11().EncoderTexture;
-				break;
-			case EVideoFrameFormat::D3D12_R8G8B8A8_UNORM:
-				UE_LOG(LogEncoderNVENC, Error, TEXT("NVENC got passed a DX12 texture which it does not support, make sure it has been shared with a DX11 context and pass that texture instead."));
-				return false; // NVENC cant encode a D3D12 texture directly this should be converted to D3D11
+				case EVideoFrameFormat::D3D11_R8G8B8A8_UNORM:
+					TextureToCompress = InBuffer->SourceFrame->GetD3D11().EncoderTexture;
+					break;
+				case EVideoFrameFormat::D3D12_R8G8B8A8_UNORM:
+					UE_LOG(LogEncoderNVENC, Error, TEXT("NVENC got passed a DX12 texture which it does not support, make sure it has been shared with a DX11 context and pass that texture instead."));
+					return false; // NVENC cant encode a D3D12 texture directly this should be converted to D3D11
 #endif
-			case EVideoFrameFormat::CUDA_R8G8B8A8_UNORM:
-				TextureToCompress = InBuffer->SourceFrame->GetCUDA().EncoderTexture;
-				break;
-			case AVEncoder::EVideoFrameFormat::Undefined:
-			default:
-				break;
+				case EVideoFrameFormat::CUDA_R8G8B8A8_UNORM:
+					TextureToCompress = InBuffer->SourceFrame->GetCUDA().EncoderTexture;
+					break;
+				case AVEncoder::EVideoFrameFormat::Undefined:
+				default:
+					break;
 			}
 
-			if(!TextureToCompress)
+			if (!TextureToCompress)
 			{
 				UE_LOG(LogEncoderNVENC, Fatal, TEXT("Got passed a null texture to encode."));
 				return false;
@@ -984,33 +975,33 @@ namespace AVEncoder
 
 			FIntPoint TextureSize(InBuffer->SourceFrame->GetWidth(), InBuffer->SourceFrame->GetHeight());
 
-			switch(InBuffer->SourceFrame->GetFormat())
+			switch (InBuffer->SourceFrame->GetFormat())
 			{
 #if PLATFORM_WINDOWS
-			case EVideoFrameFormat::D3D11_R8G8B8A8_UNORM:
-			case EVideoFrameFormat::D3D12_R8G8B8A8_UNORM:
-				CreateResourceDIRECTX(InBuffer, RegisterParam, TextureSize);
-				break;
+				case EVideoFrameFormat::D3D11_R8G8B8A8_UNORM:
+				case EVideoFrameFormat::D3D12_R8G8B8A8_UNORM:
+					CreateResourceDIRECTX(InBuffer, RegisterParam, TextureSize);
+					break;
 #endif
-			case EVideoFrameFormat::CUDA_R8G8B8A8_UNORM:
-				CreateResourceCUDAARRAY(InBuffer, RegisterParam, TextureSize);
-				break;
-			case AVEncoder::EVideoFrameFormat::Undefined:
-			default:
-				break;
+				case EVideoFrameFormat::CUDA_R8G8B8A8_UNORM:
+					CreateResourceCUDAARRAY(InBuffer, RegisterParam, TextureSize);
+					break;
+				case AVEncoder::EVideoFrameFormat::Undefined:
+				default:
+					break;
 			}
 
 			RegisterParam.resourceToRegister = InBuffer->InputTexture;
 
 			NVENCSTATUS Result = NVENC.nvEncRegisterResource(NVEncoder, &RegisterParam);
-			if(Result != NV_ENC_SUCCESS)
+			if (Result != NV_ENC_SUCCESS)
 			{
 				UE_LOG(LogEncoderNVENC, Error, TEXT("Failed to register input texture with NVENC (%s)"), *NVENC.GetErrorString(NVEncoder, Result));
 				return false;
 			}
 			InBuffer->RegisteredInput = RegisterParam.registeredResource;
 
-			if(!MapInputTexture(InBuffer))
+			if (!MapInputTexture(InBuffer))
 			{
 				UE_LOG(LogEncoderNVENC, Error, TEXT("Could not map input texture, cannot proceed with encoding this buffer."));
 				return false;
@@ -1023,20 +1014,20 @@ namespace AVEncoder
 	bool FVideoEncoderNVENC_H264::FNVENCLayer::UnregisterInputTexture(FInputOutput* InBuffer)
 	{
 		checkf(InBuffer, TEXT("Cannot unregister texture in NVENC - buffer was nullptr."));
-		if(!InBuffer)
+		if (!InBuffer)
 		{
 			return false;
 		}
 
-		if(!UnmapInputTexture(InBuffer))
+		if (!UnmapInputTexture(InBuffer))
 		{
 			UE_LOG(LogEncoderNVENC, Error, TEXT("Failed to unmap input texture."));
 		}
 
-		if(InBuffer->RegisteredInput)
+		if (InBuffer->RegisteredInput)
 		{
 			NVENCSTATUS Result = NVENC.nvEncUnregisterResource(NVEncoder, InBuffer->RegisteredInput);
-			if(Result != NV_ENC_SUCCESS)
+			if (Result != NV_ENC_SUCCESS)
 			{
 				UE_LOG(LogEncoderNVENC, Error, TEXT("Failed to unregister input texture with NVENC (%s)"), *NVENC.GetErrorString(NVEncoder, Result));
 				InBuffer->InputTexture = nullptr;
@@ -1052,17 +1043,17 @@ namespace AVEncoder
 	bool FVideoEncoderNVENC_H264::FNVENCLayer::MapInputTexture(FInputOutput* InBuffer)
 	{
 		checkf(InBuffer, TEXT("Cannot map texture in NVENC - buffer was nullptr."));
-		if(!InBuffer)
+		if (!InBuffer)
 		{
 			return false;
 		}
 
-		if(!InBuffer->MappedInput)
+		if (!InBuffer->MappedInput)
 		{
 			NVENCStruct(NV_ENC_MAP_INPUT_RESOURCE, MapInputResource);
 			MapInputResource.registeredResource = InBuffer->RegisteredInput;
 			NVENCSTATUS Result = NVENC.nvEncMapInputResource(NVEncoder, &MapInputResource);
-			if(Result != NV_ENC_SUCCESS)
+			if (Result != NV_ENC_SUCCESS)
 			{
 				UE_LOG(LogEncoderNVENC, Error, TEXT("Failed to map input texture buffer (%s)"), *NVENC.GetErrorString(NVEncoder, Result));
 				return false;
@@ -1076,15 +1067,15 @@ namespace AVEncoder
 	bool FVideoEncoderNVENC_H264::FNVENCLayer::UnmapInputTexture(FInputOutput* InBuffer)
 	{
 		checkf(InBuffer, TEXT("Cannot unmap texture in NVENC - buffer was nullptr."));
-		if(!InBuffer)
+		if (!InBuffer)
 		{
 			return false;
 		}
 
-		if(InBuffer->MappedInput)
+		if (InBuffer->MappedInput)
 		{
 			NVENCSTATUS Result = NVENC.nvEncUnmapInputResource(NVEncoder, InBuffer->MappedInput);
-			if(Result != NV_ENC_SUCCESS)
+			if (Result != NV_ENC_SUCCESS)
 			{
 				UE_LOG(LogEncoderNVENC, Error, TEXT("Failed to unmap input texture buffer (%s)"), *NVENC.GetErrorString(NVEncoder, Result));
 				InBuffer->MappedInput = nullptr;
@@ -1098,12 +1089,12 @@ namespace AVEncoder
 	bool FVideoEncoderNVENC_H264::FNVENCLayer::LockOutputBuffer(FInputOutput* InBuffer)
 	{
 		checkf(InBuffer, TEXT("Cannot lock output buffer in NVENC - buffer was nullptr."));
-		if(!InBuffer)
+		if (!InBuffer)
 		{
 			return false;
 		}
 
-		if(!InBuffer->BitstreamData)
+		if (!InBuffer->BitstreamData)
 		{
 			uint64 PreLockCycles = FPlatformTime::Cycles64();
 
@@ -1112,14 +1103,14 @@ namespace AVEncoder
 			LockBitstreamParam.outputBitstream = InBuffer->OutputBitstream;
 			NVENCSTATUS Result = NVENC.nvEncLockBitstream(NVEncoder, &LockBitstreamParam);
 
-			if(CVarNVENCEnableStats.GetValueOnAnyThread())
+			if (CVarNVENCEnableStats.GetValueOnAnyThread())
 			{
 				uint64 PostLockCycles = FPlatformTime::Cycles64();
 				double LockLatencyMs = FPlatformTime::ToMilliseconds64(PostLockCycles - PreLockCycles);
 				FNVENCStats::Get().SetLockOutputBufferLatency(LockLatencyMs);
 			}
 
-			if(Result != NV_ENC_SUCCESS)
+			if (Result != NV_ENC_SUCCESS)
 			{
 				UE_LOG(LogEncoderNVENC, Error, TEXT("Failed to lock output bitstream for NVENC encoder (%s)."), *NVENC.GetErrorString(NVEncoder, Result));
 				return false;
@@ -1139,25 +1130,25 @@ namespace AVEncoder
 	bool FVideoEncoderNVENC_H264::FNVENCLayer::UnlockOutputBuffer(FInputOutput* InBuffer)
 	{
 		checkf(InBuffer, TEXT("Cannot unregister texture NVENC - buffer was nullptr."));
-		if(!InBuffer)
+		if (!InBuffer)
 		{
 			return false;
 		}
 
-		if(InBuffer->BitstreamData)
+		if (InBuffer->BitstreamData)
 		{
 			uint64 PreUnlockCycles = FPlatformTime::Cycles64();
 
 			NVENCSTATUS Result = NVENC.nvEncUnlockBitstream(NVEncoder, InBuffer->OutputBitstream);
 
-			if(CVarNVENCEnableStats.GetValueOnAnyThread())
+			if (CVarNVENCEnableStats.GetValueOnAnyThread())
 			{
 				uint64 PostUnlockCycles = FPlatformTime::Cycles64();
 				double UnlockLatencyMs = FPlatformTime::ToMilliseconds64(PostUnlockCycles - PreUnlockCycles);
 				FNVENCStats::Get().SetUnlockOutputBufferLatency(UnlockLatencyMs);
 			}
 
-			if(Result != NV_ENC_SUCCESS)
+			if (Result != NV_ENC_SUCCESS)
 			{
 				UE_LOG(LogEncoderNVENC, Error, TEXT("Failed to unlock output bitstream for NVENC encoder (%s)."), *NVENC.GetErrorString(NVEncoder, Result));
 				return false;
@@ -1179,15 +1170,15 @@ namespace AVEncoder
 		TRefCountPtr<IDXGIAdapter> Adapter;
 
 		HRESULT Result = CreateDXGIFactory1(__uuidof(IDXGIFactory1), (void**)DXGIFactory1.GetInitReference());
-		if(Result != S_OK)
+		if (Result != S_OK)
 		{
 			UE_LOG(LogEncoderNVENC, Error, TEXT("Failed to create DX factory for NVENC."));
 			return false;
 		}
 
-		for(int GpuIndex = 0; GpuIndex < MAX_GPU_INDEXES; GpuIndex++)
+		for (int GpuIndex = 0; GpuIndex < MAX_GPU_INDEXES; GpuIndex++)
 		{
-			if((Result = DXGIFactory1->EnumAdapters(GpuIndex, Adapter.GetInitReference())) != S_OK)
+			if ((Result = DXGIFactory1->EnumAdapters(GpuIndex, Adapter.GetInitReference())) != S_OK)
 			{
 				UE_LOG(LogEncoderNVENC, Error, TEXT("Failed to enum GPU #%d for NVENC."), GpuIndex);
 				return false;
@@ -1195,13 +1186,14 @@ namespace AVEncoder
 
 			DXGI_ADAPTER_DESC AdapterDesc;
 			Adapter->GetDesc(&AdapterDesc);
-			if(AdapterDesc.VendorId != 0x10DE) // NVIDIA
+			if (AdapterDesc.VendorId != 0x10DE) // NVIDIA
 			{
 				continue;
 			}
 
-			if((Result = D3D11CreateDevice(Adapter, D3D_DRIVER_TYPE_UNKNOWN, NULL, 0, NULL, 0, D3D11_SDK_VERSION, OutEncoderDevice.GetInitReference(), NULL,
-			                               OutEncoderDeviceContext.GetInitReference())) != S_OK)
+			if ((Result = D3D11CreateDevice(Adapter, D3D_DRIVER_TYPE_UNKNOWN, NULL, 0, NULL, 0, D3D11_SDK_VERSION, OutEncoderDevice.GetInitReference(), NULL,
+					 OutEncoderDeviceContext.GetInitReference()))
+				!= S_OK)
 			{
 				UE_LOG(LogEncoderNVENC, Error, TEXT("Failed to create D3D11 device for NVENC."));
 			}
@@ -1229,7 +1221,7 @@ namespace AVEncoder
 
 		NVENCSTATUS NvResult = NVENC.nvEncOpenEncodeSessionEx(&OpenEncodeSessionExParams, &EncoderSession);
 		// UE_LOG(LogEncoderNVENC, Error, TEXT("NVENC.nvEncOpenEncodeSessionEx(&OpenEncodeSessionExParams, &EncoderSession); -> %d"), NvResult);
-		if(NvResult != NV_ENC_SUCCESS)
+		if (NvResult != NV_ENC_SUCCESS)
 		{
 			UE_LOG(LogEncoderNVENC, Error, TEXT("Unable to open NvEnc encoding session (status: %d)."), NvResult);
 			EncoderSession = nullptr;
@@ -1249,7 +1241,7 @@ namespace AVEncoder
 
 		NVENCSTATUS NvResult = NVENC.nvEncOpenEncodeSessionEx(&OpenEncodeSessionExParams, &EncoderSession);
 		//	UE_LOG(LogEncoderNVENC, Error, TEXT("NVENC.nvEncOpenEncodeSessionEx(&OpenEncodeSessionExParams, &EncoderSession); -> %d"), NvResult);
-		if(NvResult != NV_ENC_SUCCESS)
+		if (NvResult != NV_ENC_SUCCESS)
 		{
 			UE_LOG(LogEncoderNVENC, Error, TEXT("Unable to open NvEnc encoding session (status: %d)."), NvResult);
 			EncoderSession = nullptr;
@@ -1263,8 +1255,8 @@ namespace AVEncoder
 		NVENCStruct(NV_ENC_CAPS_PARAM, CapsParam);
 		CapsParam.capsToQuery = InCapsToQuery;
 		NVENCSTATUS Result = NVENC.nvEncGetEncodeCaps(InEncoder, NV_ENC_CODEC_H264_GUID, &CapsParam, &CapsValue);
-		
-		if(Result != NV_ENC_SUCCESS)
+
+		if (Result != NV_ENC_SUCCESS)
 		{
 			UE_LOG(LogEncoderNVENC, Warning, TEXT("Failed to query for NVENC capability %d (error %d)."), InCapsToQuery, Result);
 			return 0;
@@ -1280,31 +1272,31 @@ namespace AVEncoder
 
 		OutSupportedProfiles = 0;
 		NVENCSTATUS Result = NVENC.nvEncGetEncodeProfileGUIDs(InEncoder, NV_ENC_CODEC_H264_GUID, ProfileGUIDs, MaxProfileGUIDs, &NumProfileGUIDs);
-		
-		if(Result != NV_ENC_SUCCESS)
+
+		if (Result != NV_ENC_SUCCESS)
 		{
 			UE_LOG(LogEncoderNVENC, Error, TEXT("Unable to query profiles supported by NvEnc (error: %d)."), Result);
 			return false;
 		}
-		for(uint32 Index = 0; Index < NumProfileGUIDs; ++Index)
+		for (uint32 Index = 0; Index < NumProfileGUIDs; ++Index)
 		{
-			if(memcmp(&NV_ENC_H264_PROFILE_BASELINE_GUID, &ProfileGUIDs[Index], sizeof(GUID)) == 0)
+			if (memcmp(&NV_ENC_H264_PROFILE_BASELINE_GUID, &ProfileGUIDs[Index], sizeof(GUID)) == 0)
 			{
 				OutSupportedProfiles |= H264Profile_Baseline;
-				if(GetEncoderCapability(NVENC, InEncoder, NV_ENC_CAPS_SUPPORT_CONSTRAINED_ENCODING))
+				if (GetEncoderCapability(NVENC, InEncoder, NV_ENC_CAPS_SUPPORT_CONSTRAINED_ENCODING))
 				{
 					OutSupportedProfiles |= H264Profile_ConstrainedBaseline;
 				}
 			}
-			else if(memcmp(&NV_ENC_H264_PROFILE_MAIN_GUID, &ProfileGUIDs[Index], sizeof(GUID)) == 0)
+			else if (memcmp(&NV_ENC_H264_PROFILE_MAIN_GUID, &ProfileGUIDs[Index], sizeof(GUID)) == 0)
 			{
 				OutSupportedProfiles |= H264Profile_Main;
 			}
-			else if(memcmp(&NV_ENC_H264_PROFILE_HIGH_GUID, &ProfileGUIDs[Index], sizeof(GUID)) == 0)
+			else if (memcmp(&NV_ENC_H264_PROFILE_HIGH_GUID, &ProfileGUIDs[Index], sizeof(GUID)) == 0)
 			{
 				OutSupportedProfiles |= H264Profile_High;
 			}
-			else if(memcmp(&NV_ENC_H264_PROFILE_CONSTRAINED_HIGH_GUID, &ProfileGUIDs[Index], sizeof(GUID)) == 0)
+			else if (memcmp(&NV_ENC_H264_PROFILE_CONSTRAINED_HIGH_GUID, &ProfileGUIDs[Index], sizeof(GUID)) == 0)
 			{
 				OutSupportedProfiles |= H264Profile_ConstrainedHigh;
 			}
@@ -1318,29 +1310,29 @@ namespace AVEncoder
 		uint32_t InputFmtCount = 0;
 		NV_ENC_BUFFER_FORMAT InputFormats[MaxInputFmtCount];
 		NVENCSTATUS Result = NVENC.nvEncGetInputFormats(InEncoder, NV_ENC_CODEC_H264_GUID, InputFormats, MaxInputFmtCount, &InputFmtCount);
-		if(Result != NV_ENC_SUCCESS)
+		if (Result != NV_ENC_SUCCESS)
 		{
 			UE_LOG(LogEncoderNVENC, Error, TEXT("Unable to query input formats supported by NvEnc (error: %d)."), Result);
 			return false;
 		}
 
-		for(uint32_t Index = 0; Index < InputFmtCount; ++Index)
+		for (uint32_t Index = 0; Index < InputFmtCount; ++Index)
 		{
-			switch(InputFormats[Index])
+			switch (InputFormats[Index])
 			{
-			case NV_ENC_BUFFER_FORMAT_IYUV:
-				break;
-			case NV_ENC_BUFFER_FORMAT_NV12:
-				break;
-			case NV_ENC_BUFFER_FORMAT_ARGB:
-				break;
-			case NV_ENC_BUFFER_FORMAT_ABGR:
+				case NV_ENC_BUFFER_FORMAT_IYUV:
+					break;
+				case NV_ENC_BUFFER_FORMAT_NV12:
+					break;
+				case NV_ENC_BUFFER_FORMAT_ARGB:
+					break;
+				case NV_ENC_BUFFER_FORMAT_ABGR:
 #if PLATFORM_WINDOWS
-				OutSupportedInputFormats.Push(EVideoFrameFormat::D3D11_R8G8B8A8_UNORM);
-				OutSupportedInputFormats.Push(EVideoFrameFormat::D3D12_R8G8B8A8_UNORM);
+					OutSupportedInputFormats.Push(EVideoFrameFormat::D3D11_R8G8B8A8_UNORM);
+					OutSupportedInputFormats.Push(EVideoFrameFormat::D3D12_R8G8B8A8_UNORM);
 #endif
-				OutSupportedInputFormats.Push(EVideoFrameFormat::CUDA_R8G8B8A8_UNORM);
-				break;
+					OutSupportedInputFormats.Push(EVideoFrameFormat::CUDA_R8G8B8A8_UNORM);
+					break;
 			}
 		}
 		return true;
@@ -1358,25 +1350,25 @@ namespace AVEncoder
 		TRefCountPtr<ID3D11Device> EncoderDevice;
 		TRefCountPtr<ID3D11DeviceContext> EncoderDeviceContext;
 
-		if(!CreateEncoderDevice(EncoderDevice, EncoderDeviceContext))
+		if (!CreateEncoderDevice(EncoderDevice, EncoderDeviceContext))
 		{
 			bSuccess = false;
 		}
 
-		if((EncoderSession = CreateEncoderSession(NVENC, EncoderDevice)) == nullptr)
+		if ((EncoderSession = CreateEncoderSession(NVENC, EncoderDevice)) == nullptr)
 		{
 			bSuccess = false;
 		}
 #endif
 		// if we dont already have an encoder session try with CUDA if its avaliable
-		if(!EncoderSession && FModuleManager::GetModulePtr<FCUDAModule>("CUDA"))
+		if (!EncoderSession && FModuleManager::GetModulePtr<FCUDAModule>("CUDA"))
 		{
 			FCUDAModule& CUDAModule = FModuleManager::GetModuleChecked<FCUDAModule>("CUDA");
-			if(CUDAModule.IsAvailable())
+			if (CUDAModule.IsAvailable())
 			{
 				EncoderSession = CreateEncoderSession(NVENC, CUDAModule.GetCudaContext());
 				bSuccess = EncoderSession != nullptr;
-			} 
+			}
 			else
 			{
 				bSuccess = false;
@@ -1384,7 +1376,7 @@ namespace AVEncoder
 		}
 
 		// if we dont have a session by now opt out. this will cause NVENC to not register
-		if(!EncoderSession || !bSuccess)
+		if (!EncoderSession || !bSuccess)
 		{
 			return false;
 		}
@@ -1395,10 +1387,11 @@ namespace AVEncoder
 
 		int LevelMax = GetEncoderCapability(NVENC, EncoderSession, NV_ENC_CAPS_LEVEL_MAX);
 		int LevelMin = GetEncoderCapability(NVENC, EncoderSession, NV_ENC_CAPS_LEVEL_MIN);
-		if(LevelMin > 0 && LevelMax > 0 && LevelMax >= LevelMin)
+		if (LevelMin > 0 && LevelMax > 0 && LevelMax >= LevelMin)
 		{
 			EncoderInfo.H264.MinLevel = (LevelMin > 9) ? LevelMin : 9;
-			EncoderInfo.H264.MaxLevel = (LevelMax < 9) ? 9 : (LevelMax > NV_ENC_LEVEL_H264_52) ? NV_ENC_LEVEL_H264_52 : LevelMax;
+			EncoderInfo.H264.MaxLevel = (LevelMax < 9) ? 9 : (LevelMax > NV_ENC_LEVEL_H264_52) ? NV_ENC_LEVEL_H264_52
+																							   : LevelMax;
 		}
 		else
 		{
@@ -1406,13 +1399,13 @@ namespace AVEncoder
 			bSuccess = false;
 		}
 
-		if(!GetEncoderSupportedProfiles(NVENC, EncoderSession, EncoderInfo.H264.SupportedProfiles) || !GetEncoderSupportedInputFormats(NVENC, EncoderSession, EncoderInfo.SupportedInputFormats))
+		if (!GetEncoderSupportedProfiles(NVENC, EncoderSession, EncoderInfo.H264.SupportedProfiles) || !GetEncoderSupportedInputFormats(NVENC, EncoderSession, EncoderInfo.SupportedInputFormats))
 		{
 			bSuccess = false;
 		}
 
 		// destroy encoder session
-		if(EncoderSession)
+		if (EncoderSession)
 		{
 			NVENC.nvEncDestroyEncoder(EncoderSession);
 		}
diff --git a/Engine/Plugins/Media/HardwareEncoders/Source/EncoderNVENC/Private/NVENC_EncoderH264.h b/Engine/Plugins/Media/HardwareEncoders/Source/EncoderNVENC/Private/NVENC_EncoderH264.h
--- a/Engine/Plugins/Media/HardwareEncoders/Source/EncoderNVENC/Private/NVENC_EncoderH264.h
+++ b/Engine/Plugins/Media/HardwareEncoders/Source/EncoderNVENC/Private/NVENC_EncoderH264.h
@@ -3,112 +3,110 @@
 #pragma once
 
 #include "NVENC_Common.h"
-
 #include "VideoEncoderFactory.h"
 #include "VideoEncoderInputImpl.h"
-
 #include "HAL/Event.h"
 
 namespace AVEncoder
 {
-    class FVideoEncoderNVENC_H264 : public FVideoEncoder
-    {
+	class FVideoEncoderNVENC_H264 : public FVideoEncoder
+	{
 
-    public:
-        virtual ~FVideoEncoderNVENC_H264() override;
+	public:
+		virtual ~FVideoEncoderNVENC_H264() override; 
 
-        bool Setup(TSharedRef<FVideoEncoderInput> InputFrameFactory, const FLayerConfig& InitConfig) override;
-        void Shutdown() override;
+		bool Setup(TSharedRef<FVideoEncoderInput> InputFrameFactory, const FLayerConfig& InitConfig) override;
+		void Shutdown() override;
 
-        // Query whether or not encoder is supported and available
-        static bool GetIsAvailable(const FVideoEncoderInput& InFrameFactory, FVideoEncoderInfo &OutEncoderInfo);
+		// Query whether or not encoder is supported and available
+		static bool GetIsAvailable(const FVideoEncoderInput& InFrameFactory, FVideoEncoderInfo& OutEncoderInfo);
 
-        // Register encoder with video encoder factory
-        static void Register(FVideoEncoderFactory &InFactory);
+		// Register encoder with video encoder factory
+		static void Register(FVideoEncoderFactory& InFactory);
 
-        void Encode(FVideoEncoderInputFrame const* InFrame, const FEncodeOptions& EncodeOptions) override;
-        void Flush();
+		void Encode(FVideoEncoderInputFrame const* InFrame, const FEncodeOptions& EncodeOptions) override;
+		void Flush();
 
-    protected:
-        FLayer *CreateLayer(uint32 InLayerIndex, const FLayerConfig &InLayerConfig) override;
+	protected:
+		FLayer* CreateLayer(uint32 InLayerIndex, const FLayerConfig& InLayerConfig) override;
 		void DestroyLayer(FLayer* layer) override;
 
-    private:
-        FVideoEncoderNVENC_H264();
-
-        class FNVENCLayer : public FLayer
-        {
-
-        public:
-
-            struct FInputOutput
-            {
-                const AVEncoder::FVideoEncoderInputFrameImpl* SourceFrame = nullptr;
-                void*  InputTexture = nullptr;
-                uint32 Width = 0;
-                uint32 Height = 0;
-                uint32 Pitch = 0;
-                NV_ENC_BUFFER_FORMAT BufferFormat = NV_ENC_BUFFER_FORMAT_UNDEFINED;
-                NV_ENC_REGISTERED_PTR RegisteredInput = nullptr;
-                NV_ENC_INPUT_PTR MappedInput = nullptr;
-                NV_ENC_PIC_PARAMS PicParams = {};
-                NV_ENC_OUTPUT_PTR OutputBitstream = nullptr;
-                const void *BitstreamData = nullptr;
-                uint32 BitstreamDataSize = 0;
-                NV_ENC_PIC_TYPE PictureType = NV_ENC_PIC_TYPE_UNKNOWN;
-                uint32 FrameAvgQP = 0;
-                uint64 TimeStamp;
-                uint64 SubmitTimeCycles;
-            };
-
-            FNVENCLayer(uint32 layerIdx, FLayerConfig const& config, FVideoEncoderNVENC_H264& encoder);
-            ~FNVENCLayer();
-            bool Setup();
-            bool CreateSession();
-            bool CreateInitialConfig();
-            int GetCapability(NV_ENC_CAPS CapsToQuery) const;
-            FString GetError(NVENCSTATUS ForStatus) const;
+	private:
+		FVideoEncoderNVENC_H264();
+
+		class FNVENCLayer : public FLayer
+		{
+
+		public:
+			struct FInputOutput
+			{
+				const AVEncoder::FVideoEncoderInputFrameImpl* SourceFrame = nullptr;
+				void* InputTexture = nullptr;
+				uint32 Width = 0;
+				uint32 Height = 0;
+				uint32 Pitch = 0;
+				NV_ENC_BUFFER_FORMAT BufferFormat = NV_ENC_BUFFER_FORMAT_UNDEFINED;
+				NV_ENC_REGISTERED_PTR RegisteredInput = nullptr;
+				NV_ENC_INPUT_PTR MappedInput = nullptr;
+				NV_ENC_PIC_PARAMS PicParams = {};
+				NV_ENC_OUTPUT_PTR OutputBitstream = nullptr;
+				const void* BitstreamData = nullptr;
+				uint32 BitstreamDataSize = 0;
+				NV_ENC_PIC_TYPE PictureType = NV_ENC_PIC_TYPE_UNKNOWN;
+				uint32 FrameAvgQP = 0;
+				uint64 TimeStamp;
+				uint64 SubmitTimeCycles;
+			};
+
+			FNVENCLayer(uint32 layerIdx, FLayerConfig const& config, FVideoEncoderNVENC_H264& encoder);
+			~FNVENCLayer();
+			bool Setup();
+			bool CreateSession();
+			bool CreateInitialConfig();
+			int GetCapability(NV_ENC_CAPS CapsToQuery) const;
+			FString GetError(NVENCSTATUS ForStatus) const;
 			void MaybeReconfigure();
 			void UpdateConfig();
-            void UpdateLastEncodedQP(uint32 InLastEncodedQP);
-            void Encode(FVideoEncoderInputFrame const* InFrame, const FEncodeOptions& EncodeOptions);
-            void EncodeBuffer(FInputOutput* Buffer);
-            void ProcessEncodedBuffer(FInputOutput* Buffer);
-            void Flush();
-            void Shutdown();
-            void UpdateBitrate(uint32 InMaxBitRate, uint32 InTargetBitRate);
-            void UpdateResolution(uint32 InMaxBitRate, uint32 InTargetBitRate);
-            FInputOutput* GetOrCreateBuffer(const FVideoEncoderInputFrameImpl* InFrame);
-            FInputOutput* CreateBuffer();
-            void DestroyBuffer(FInputOutput* InBuffer);
-            bool RegisterInputTexture(FInputOutput* InBuffer);
-            bool UnregisterInputTexture(FInputOutput* InBuffer);
-            bool MapInputTexture(FInputOutput* InBuffer);
-            bool UnmapInputTexture(FInputOutput* InBuffer);
-            bool LockOutputBuffer(FInputOutput* InBuffer);
-            bool UnlockOutputBuffer(FInputOutput* InBuffer);
-            void CreateResourceDIRECTX(FInputOutput* InBuffer, NV_ENC_REGISTER_RESOURCE &RegisterParam, FIntPoint TextureSize);
-            void CreateResourceCUDAARRAY(FInputOutput* InBuffer, NV_ENC_REGISTER_RESOURCE &RegisterParam, FIntPoint TextureSize);
-
-        public:
-            FVideoEncoderNVENC_H264 &Encoder;
-            FNVENCCommon &NVENC;
-            GUID CodecGUID;
-            uint32 LayerIndex;
-            void *NVEncoder = nullptr;
-            NV_ENC_INITIALIZE_PARAMS EncoderInitParams;
-            NV_ENC_CONFIG EncoderConfig;
-            FDateTime LastKeyFrameTime = 0;
-            bool bForceNextKeyframe = false;
-            uint32 LastEncodedQP = 0;
-            TMap<const AVEncoder::FVideoEncoderInputFrame*, FInputOutput*> FrameToBufferMapping;
-        };
-
-        FNVENCCommon &NVENC;
-        EVideoFrameFormat FrameFormat = EVideoFrameFormat::Undefined;
-
-        // this could be a TRefCountPtr<ID3D11Device>, CUcontext or void*
-        void* EncoderDevice;
-    };
-
-}
+			void UpdateLastEncodedQP(uint32 InLastEncodedQP);
+			void Encode(FVideoEncoderInputFrame const* InFrame, const FEncodeOptions& EncodeOptions);
+			void EncodeBuffer(FInputOutput* Buffer);
+			void ProcessEncodedBuffer(FInputOutput* Buffer);
+			void Flush();
+			void Shutdown();
+			void UpdateBitrate(uint32 InMaxBitRate, uint32 InTargetBitRate);
+			void UpdateResolution(uint32 InMaxBitRate, uint32 InTargetBitRate);
+			FInputOutput* GetOrCreateBuffer(const FVideoEncoderInputFrameImpl* InFrame);
+			FInputOutput* CreateBuffer();
+			void DestroyBuffer(FInputOutput* InBuffer);
+			bool RegisterInputTexture(FInputOutput* InBuffer);
+			bool UnregisterInputTexture(FInputOutput* InBuffer);
+			bool MapInputTexture(FInputOutput* InBuffer);
+			bool UnmapInputTexture(FInputOutput* InBuffer);
+			bool LockOutputBuffer(FInputOutput* InBuffer);
+			bool UnlockOutputBuffer(FInputOutput* InBuffer);
+			void CreateResourceDIRECTX(FInputOutput* InBuffer, NV_ENC_REGISTER_RESOURCE& RegisterParam, FIntPoint TextureSize);
+			void CreateResourceCUDAARRAY(FInputOutput* InBuffer, NV_ENC_REGISTER_RESOURCE& RegisterParam, FIntPoint TextureSize);
+
+		public:
+			FVideoEncoderNVENC_H264& Encoder;
+			FNVENCCommon& NVENC;
+			GUID CodecGUID;
+			uint32 LayerIndex;
+			void* NVEncoder = nullptr;
+			NV_ENC_INITIALIZE_PARAMS EncoderInitParams;
+			NV_ENC_CONFIG EncoderConfig;
+			FDateTime LastKeyFrameTime = 0;
+			bool bForceNextKeyframe = false;
+			uint32 LastEncodedQP = 0;
+			FInputOutput* InputOutputBuffer = nullptr;
+			FThreadSafeBool bIsProcessingFrame = false;
+		};
+
+		FNVENCCommon& NVENC;
+		EVideoFrameFormat FrameFormat = EVideoFrameFormat::Undefined;
+
+		// this could be a TRefCountPtr<ID3D11Device>, CUcontext or void*
+		void* EncoderDevice;
+	};
+
+} // namespace AVEncoder
diff --git a/Engine/Plugins/Media/PixelStreaming/Config/BasePixelStreaming.ini b/Engine/Plugins/Media/PixelStreaming/Config/BasePixelStreaming.ini
--- a/Engine/Plugins/Media/PixelStreaming/Config/BasePixelStreaming.ini
+++ b/Engine/Plugins/Media/PixelStreaming/Config/BasePixelStreaming.ini
@@ -1,5 +1,8 @@
 [/Script/PixelStreaming.PixelStreamingSettings]
-PixelStreamerDefaultCursorClassName=/PixelStreaming/DefaultCursor.DefaultCursor_C
-PixelStreamerTextEditBeamCursorClassName=/PixelStreaming/TextEditBeamCursor.TextEditBeamCursor_C
-PixelStreamerHiddenCursorClassName=/PixelStreaming/HiddenCursor.HiddenCursor_C
-bPixelStreamerMouseAlwaysAttached=true
\ No newline at end of file
+DefaultCursorClassName=/PixelStreaming/DefaultCursor.DefaultCursor_C
+TextEditBeamCursorClassName=/PixelStreaming/TextEditBeamCursor.TextEditBeamCursor_C
+HiddenCursorClassName=/PixelStreaming/HiddenCursor.HiddenCursor_C
+bMouseAlwaysAttached=true
+
+[CoreRedirects]
++ClassRedirects=(OldName="/Script/PixelStreaming.PixelStreamerInputComponent", NewName="/Script/PixelStreaming.PixelStreamingInput")
\ No newline at end of file
diff --git a/Engine/Plugins/Media/PixelStreaming/Source/PixelStreaming/Private/AudioCapturer.cpp b/Engine/Plugins/Media/PixelStreaming/Source/PixelStreaming/Private/AudioCapturer.cpp
--- a/Engine/Plugins/Media/PixelStreaming/Source/PixelStreaming/Private/AudioCapturer.cpp
+++ b/Engine/Plugins/Media/PixelStreaming/Source/PixelStreaming/Private/AudioCapturer.cpp
@@ -5,8 +5,8 @@
 #include "SampleBuffer.h"
 #include "Engine/GameEngine.h"
 
-DECLARE_LOG_CATEGORY_EXTERN(LogAudioCapturer, Log, All);
-DEFINE_LOG_CATEGORY(LogAudioCapturer);
+DECLARE_LOG_CATEGORY_EXTERN(LogPixelStreamingAudioCapturer, Log, All);
+DEFINE_LOG_CATEGORY(LogPixelStreamingAudioCapturer);
 
 // These are copied from webrtc internals
 #define CHECKinitialized_() \
@@ -24,10 +24,10 @@ DEFINE_LOG_CATEGORY(LogAudioCapturer);
 		};                       \
 	}
 
-constexpr int FAudioCapturer::SampleRate;
-constexpr int FAudioCapturer::NumChannels;
+constexpr int UE::PixelStreaming::FAudioCapturer::SampleRate;
+constexpr int UE::PixelStreaming::FAudioCapturer::NumChannels;
 
-void FAudioCapturer::OnNewSubmixBuffer(const USoundSubmix* OwningSubmix, float* AudioData, int32 NumSamples, int32 InNumChannels, const int32 InSampleRate, double AudioClock)
+void UE::PixelStreaming::FAudioCapturer::OnNewSubmixBuffer(const USoundSubmix* OwningSubmix, float* AudioData, int32 NumSamples, int32 InNumChannels, const int32 InSampleRate, double AudioClock)
 {
 	if (!(bInitialized && bRecordingInitialized))
 	{
@@ -41,12 +41,12 @@ void FAudioCapturer::OnNewSubmixBuffer(const USoundSubmix* OwningSubmix, float*
 		if (!bFormatChecked)
 		{
 			bFormatChecked = true;
-			UE_LOG(LogAudioCapturer, Error, TEXT("Audio samplerate needs to be 48000hz"));
+			UE_LOG(LogPixelStreamingAudioCapturer, Error, TEXT("Audio samplerate needs to be 48000hz"));
 		}
 		return;
 	}
 
-	UE_LOG(LogAudioCapturer, VeryVerbose, TEXT("captured %d samples, %dc, %dHz"), NumSamples, NumChannels, SampleRate);
+	UE_LOG(LogPixelStreamingAudioCapturer, VeryVerbose, TEXT("captured %d samples, %dc, %dHz"), NumSamples, NumChannels, SampleRate);
 
 	Audio::TSampleBuffer<float> Buffer(AudioData, NumSamples, InNumChannels, SampleRate);
 	// Mix to stereo if required, since PixelStreaming only accepts stereo at the moment
@@ -78,7 +78,7 @@ void FAudioCapturer::OnNewSubmixBuffer(const USoundSubmix* OwningSubmix, float*
 			{
 				DeviceBuffer->SetRecordedBuffer(RecordingBuffer.GetData(), BytesPer10Ms / (sizeof(uint16) * NumChannels));
 				DeviceBuffer->DeliverRecordedData();
-				UE_LOG(LogAudioCapturer, VeryVerbose, TEXT("passed %d bytes"), BytesPer10Ms);
+				UE_LOG(LogPixelStreamingAudioCapturer, VeryVerbose, TEXT("passed %d bytes"), BytesPer10Ms);
 			}
 		}
 
@@ -86,13 +86,13 @@ void FAudioCapturer::OnNewSubmixBuffer(const USoundSubmix* OwningSubmix, float*
 	}
 }
 
-int32 FAudioCapturer::ActiveAudioLayer(AudioLayer* audioLayer) const
+int32 UE::PixelStreaming::FAudioCapturer::ActiveAudioLayer(AudioLayer* audioLayer) const
 {
 	*audioLayer = AudioDeviceModule::kDummyAudio;
 	return 0;
 }
 
-int32 FAudioCapturer::RegisterAudioCallback(webrtc::AudioTransport* audioCallback)
+int32 UE::PixelStreaming::FAudioCapturer::RegisterAudioCallback(webrtc::AudioTransport* audioCallback)
 {
 	FScopeLock Lock(&DeviceBufferCS);
 	DeviceBuffer->RegisterAudioCallback(audioCallback);
@@ -100,7 +100,7 @@ int32 FAudioCapturer::RegisterAudioCallback(webrtc::AudioTransport* audioCallbac
 	return 0;
 }
 
-int32 FAudioCapturer::Init()
+int32 UE::PixelStreaming::FAudioCapturer::Init()
 {
 	if (bInitialized)
 		return 0;
@@ -121,19 +121,19 @@ int32 FAudioCapturer::Init()
 	FAudioDeviceHandle AudioDevice = GEngine->GetMainAudioDevice();
 	if (!AudioDevice)
 	{
-		UE_LOG(LogAudioCapturer, Warning, TEXT("No audio device"));
+		UE_LOG(LogPixelStreamingAudioCapturer, Warning, TEXT("No audio device"));
 		return -1;
 	}
 
 	bInitialized = true;
 	AudioDevice->RegisterSubmixBufferListener(this);
 
-	UE_LOG(LogAudioCapturer, Verbose, TEXT("Init"));
+	UE_LOG(LogPixelStreamingAudioCapturer, Verbose, TEXT("Init"));
 
 	return 0;
 }
 
-int32 FAudioCapturer::Terminate()
+int32 UE::PixelStreaming::FAudioCapturer::Terminate()
 {
 	if (!bInitialized)
 		return 0;
@@ -158,91 +158,91 @@ int32 FAudioCapturer::Terminate()
 		DeviceBuffer.Reset();
 	}
 
-	UE_LOG(LogAudioCapturer, Verbose, TEXT("Terminate"));
+	UE_LOG(LogPixelStreamingAudioCapturer, Verbose, TEXT("Terminate"));
 
 	return 0;
 }
 
-bool FAudioCapturer::Initialized() const
+bool UE::PixelStreaming::FAudioCapturer::Initialized() const
 {
 	return bInitialized;
 }
 
-int16 FAudioCapturer::PlayoutDevices()
+int16 UE::PixelStreaming::FAudioCapturer::PlayoutDevices()
 {
 	CHECKinitialized_();
 	return -1;
 }
 
-int16 FAudioCapturer::RecordingDevices()
+int16 UE::PixelStreaming::FAudioCapturer::RecordingDevices()
 {
 	CHECKinitialized_();
 	return -1;
 }
 
-int32 FAudioCapturer::PlayoutDeviceName(
+int32 UE::PixelStreaming::FAudioCapturer::PlayoutDeviceName(
 	uint16 index, char name[webrtc::kAdmMaxDeviceNameSize], char guid[webrtc::kAdmMaxGuidSize])
 {
 	CHECKinitialized_();
 	return -1;
 }
 
-int32 FAudioCapturer::RecordingDeviceName(
+int32 UE::PixelStreaming::FAudioCapturer::RecordingDeviceName(
 	uint16 index, char name[webrtc::kAdmMaxDeviceNameSize], char guid[webrtc::kAdmMaxGuidSize])
 {
 	CHECKinitialized_();
 	return -1;
 }
 
-int32 FAudioCapturer::SetPlayoutDevice(uint16 index)
+int32 UE::PixelStreaming::FAudioCapturer::SetPlayoutDevice(uint16 index)
 {
 	CHECKinitialized_();
 	return 0;
 }
 
-int32 FAudioCapturer::SetPlayoutDevice(WindowsDeviceType device)
+int32 UE::PixelStreaming::FAudioCapturer::SetPlayoutDevice(WindowsDeviceType device)
 {
 	CHECKinitialized_();
 	return 0;
 }
 
-int32 FAudioCapturer::SetRecordingDevice(uint16 index)
+int32 UE::PixelStreaming::FAudioCapturer::SetRecordingDevice(uint16 index)
 {
 	CHECKinitialized_();
 	return 0;
 }
 
-int32 FAudioCapturer::SetRecordingDevice(WindowsDeviceType device)
+int32 UE::PixelStreaming::FAudioCapturer::SetRecordingDevice(WindowsDeviceType device)
 {
 	CHECKinitialized_();
 	return 0;
 }
 
-int32 FAudioCapturer::PlayoutIsAvailable(bool* available)
+int32 UE::PixelStreaming::FAudioCapturer::PlayoutIsAvailable(bool* available)
 {
 	CHECKinitialized_();
 	return -1;
 }
 
-int32 FAudioCapturer::InitPlayout()
+int32 UE::PixelStreaming::FAudioCapturer::InitPlayout()
 {
 	CHECKinitialized_();
 	return -1;
 }
 
-bool FAudioCapturer::PlayoutIsInitialized() const
+bool UE::PixelStreaming::FAudioCapturer::PlayoutIsInitialized() const
 {
 	CHECKinitialized__BOOL();
 	return false;
 }
 
-int32 FAudioCapturer::RecordingIsAvailable(bool* available)
+int32 UE::PixelStreaming::FAudioCapturer::RecordingIsAvailable(bool* available)
 {
 	CHECKinitialized_();
 	return -1;
 }
 
-int32 FAudioCapturer::InitRecording()
+int32 UE::PixelStreaming::FAudioCapturer::InitRecording()
 {
 	CHECKinitialized_();
 
@@ -258,104 +258,104 @@ int32 FAudioCapturer::InitRecording()
 	return 0;
 }
 
-bool FAudioCapturer::RecordingIsInitialized() const
+bool UE::PixelStreaming::FAudioCapturer::RecordingIsInitialized() const
 {
 	CHECKinitialized__BOOL();
 	return bRecordingInitialized == true;
 }
 
-int32 FAudioCapturer::StartPlayout()
+int32 UE::PixelStreaming::FAudioCapturer::StartPlayout()
 {
 	CHECKinitialized_();
 	return -1;
 }
 
-int32 FAudioCapturer::StopPlayout()
+int32 UE::PixelStreaming::FAudioCapturer::StopPlayout()
 {
 	CHECKinitialized_();
 	return -1;
 }
 
-bool FAudioCapturer::Playing() const
+bool UE::PixelStreaming::FAudioCapturer::Playing() const
 {
 	CHECKinitialized__BOOL();
 	return false;
 }
 
-int32 FAudioCapturer::StartRecording()
+int32 UE::PixelStreaming::FAudioCapturer::StartRecording()
 {
 	CHECKinitialized_();
 	return -1;
 }
 
-int32 FAudioCapturer::StopRecording()
+int32 UE::PixelStreaming::FAudioCapturer::StopRecording()
 {
 	CHECKinitialized_();
 	return -1;
 }
 
-bool FAudioCapturer::Recording() const
+bool UE::PixelStreaming::FAudioCapturer::Recording() const
 {
 	CHECKinitialized__BOOL();
 	return bRecordingInitialized;
 }
 
-int32 FAudioCapturer::InitSpeaker()
+int32 UE::PixelStreaming::FAudioCapturer::InitSpeaker()
 {
 	CHECKinitialized_();
 	return -1;
 }
 
-bool FAudioCapturer::SpeakerIsInitialized() const
+bool UE::PixelStreaming::FAudioCapturer::SpeakerIsInitialized() const
 {
 	CHECKinitialized__BOOL();
 	return false;
 }
 
-int32 FAudioCapturer::InitMicrophone()
+int32 UE::PixelStreaming::FAudioCapturer::InitMicrophone()
 {
 	CHECKinitialized_();
 	return 0;
 }
 
-bool FAudioCapturer::MicrophoneIsInitialized() const
+bool UE::PixelStreaming::FAudioCapturer::MicrophoneIsInitialized() const
 {
 	CHECKinitialized__BOOL();
 	return true;
 }
 
-int32 FAudioCapturer::StereoPlayoutIsAvailable(bool* available) const
+int32 UE::PixelStreaming::FAudioCapturer::StereoPlayoutIsAvailable(bool* available) const
 {
 	CHECKinitialized_();
 	return -1;
 }
 
-int32 FAudioCapturer::SetStereoPlayout(bool enable)
+int32 UE::PixelStreaming::FAudioCapturer::SetStereoPlayout(bool enable)
 {
 	CHECKinitialized_();
 	return -1;
 }
 
-int32 FAudioCapturer::StereoPlayout(bool* enabled) const
+int32 UE::PixelStreaming::FAudioCapturer::StereoPlayout(bool* enabled) const
 {
 	CHECKinitialized_();
 	return -1;
 }
 
-int32 FAudioCapturer::StereoRecordingIsAvailable(bool* available) const
+int32 UE::PixelStreaming::FAudioCapturer::StereoRecordingIsAvailable(bool* available) const
 {
 	CHECKinitialized_();
 	*available = true;
 	return 0;
 }
 
-int32 FAudioCapturer::SetStereoRecording(bool enable)
+int32 UE::PixelStreaming::FAudioCapturer::SetStereoRecording(bool enable)
 {
 	CHECKinitialized_();
 	return 0;
 }
 
-int32 FAudioCapturer::StereoRecording(bool* enabled) const
+int32 UE::PixelStreaming::FAudioCapturer::StereoRecording(bool* enabled) const
 {
 	CHECKinitialized_();
 	*enabled = true;
diff --git a/Engine/Plugins/Media/PixelStreaming/Source/PixelStreaming/Private/AudioCapturer.h b/Engine/Plugins/Media/PixelStreaming/Source/PixelStreaming/Private/AudioCapturer.h
--- a/Engine/Plugins/Media/PixelStreaming/Source/PixelStreaming/Private/AudioCapturer.h
+++ b/Engine/Plugins/Media/PixelStreaming/Source/PixelStreaming/Private/AudioCapturer.h
@@ -6,187 +6,191 @@
 #include "AudioDevice.h"
 #include "HAL/CriticalSection.h"
 
-// This class is only used if -PixelStreamingWebRTCUseLegacyAudioDevice or related CVar are used.
-// This class will likely be removed in the future once the new ADM is confirmed stable.
-class FAudioCapturer : public webrtc::AudioDeviceModule, public ISubmixBufferListener
-{
-private:
-	// ISubmixBufferListener interface
-	void OnNewSubmixBuffer(const USoundSubmix* OwningSubmix, float* AudioData, int32 NumSamples, int32 NumChannels, const int32 SampleRate, double AudioClock) override;
-
-	//
-	// webrtc::AudioDeviceModule interface
-	//
-	int32 ActiveAudioLayer(AudioLayer* audioLayer) const override;
-	int32 RegisterAudioCallback(webrtc::AudioTransport* audioCallback) override;
-
-	// Main initialization and termination
-	int32 Init() override;
-	int32 Terminate() override;
-	bool Initialized() const override;
-
-	// Device enumeration
-	int16 PlayoutDevices() override;
-	int16 RecordingDevices() override;
-	int32 PlayoutDeviceName(uint16 index, char name[webrtc::kAdmMaxDeviceNameSize], char guid[webrtc::kAdmMaxGuidSize]) override;
-	int32 RecordingDeviceName(uint16 index, char name[webrtc::kAdmMaxDeviceNameSize], char guid[webrtc::kAdmMaxGuidSize]) override;
-
-	// Device selection
-	int32 SetPlayoutDevice(uint16 index) override;
-	int32 SetPlayoutDevice(WindowsDeviceType device) override;
-	int32 SetRecordingDevice(uint16 index) override;
-	int32 SetRecordingDevice(WindowsDeviceType device) override;
-
-	// Audio transport initialization
-	int32 PlayoutIsAvailable(bool* available) override;
-	int32 InitPlayout() override;
-	bool PlayoutIsInitialized() const override;
-	int32 RecordingIsAvailable(bool* available) override;
-	int32 InitRecording() override;
-	bool RecordingIsInitialized() const override;
-
-	// Audio transport control
-	virtual int32 StartPlayout() override;
-	virtual int32 StopPlayout() override;
-	virtual bool Playing() const override;
-	virtual int32 StartRecording() override;
-	virtual int32 StopRecording() override;
-	virtual bool Recording() const override;
-
-	// Audio mixer initialization
-	virtual int32 InitSpeaker() override;
-	virtual bool SpeakerIsInitialized() const override;
-	virtual int32 InitMicrophone() override;
-	virtual bool MicrophoneIsInitialized() const override;
-
-	// Speaker volume controls
-	virtual int32 SpeakerVolumeIsAvailable(bool* available) override
-	{
-		return -1;
-	}
-	virtual int32 SetSpeakerVolume(uint32 volume) override
-	{
-		return -1;
-	}
-	virtual int32 SpeakerVolume(uint32* volume) const override
-	{
-		return -1;
-	}
-	virtual int32 MaxSpeakerVolume(uint32* maxVolume) const override
-	{
-		return -1;
-	}
-	virtual int32 MinSpeakerVolume(uint32* minVolume) const override
-	{
-		return -1;
-	}
-
-	// Microphone volume controls
-	virtual int32 MicrophoneVolumeIsAvailable(bool* available) override
-	{
-		return -1;
-	}
-	virtual int32 SetMicrophoneVolume(uint32 volume) override
-	{
-		return -1;
-	}
-	virtual int32 MicrophoneVolume(uint32* volume) const override
-	{
-		return -1;
-	}
-	virtual int32 MaxMicrophoneVolume(uint32* maxVolume) const override
-	{
-		return -1;
-	}
-	virtual int32 MinMicrophoneVolume(uint32* minVolume) const override
-	{
-		return -1;
-	}
-
-	// Speaker mute control
-	virtual int32 SpeakerMuteIsAvailable(bool* available) override
-	{
-		return -1;
-	}
-	virtual int32 SetSpeakerMute(bool enable) override
-	{
-		return -1;
-	}
-	virtual int32 SpeakerMute(bool* enabled) const override
-	{
-		return -1;
-	}
-
-	// Microphone mute control
-	virtual int32 MicrophoneMuteIsAvailable(bool* available) override
-	{
-		return -1;
-	}
-	virtual int32 SetMicrophoneMute(bool enable) override
-	{
-		return -1;
-	}
-	virtual int32 MicrophoneMute(bool* enabled) const override
-	{
-		return -1;
-	}
-
-	// Stereo support
-	virtual int32 StereoPlayoutIsAvailable(bool* available) const override;
-	virtual int32 SetStereoPlayout(bool enable) override;
-	virtual int32 StereoPlayout(bool* enabled) const override;
-	virtual int32 StereoRecordingIsAvailable(bool* available) const override;
-	virtual int32 SetStereoRecording(bool enable) override;
-	virtual int32 StereoRecording(bool* enabled) const override;
-
-	// Playout delay
-	virtual int32 PlayoutDelay(uint16* delayMS) const override
-	{
-		return -1;
-	}
-
-	// Only supported on Android.
-	virtual bool BuiltInAECIsAvailable() const override
-	{
-		return false;
-	}
-	virtual bool BuiltInAGCIsAvailable() const override
-	{
-		return false;
-	}
-	virtual bool BuiltInNSIsAvailable() const override
-	{
-		return false;
-	}
-
-	// Enables the built-in audio effects. Only supported on Android.
-	virtual int32 EnableBuiltInAEC(bool enable) override
-	{
-		return -1;
-	}
-	virtual int32 EnableBuiltInAGC(bool enable) override
-	{
-		return -1;
-	}
-	virtual int32 EnableBuiltInNS(bool enable) override
-	{
-		return -1;
-	}
-
-	FThreadSafeBool bInitialized = false;
-	TArray<int16> PCM16;
-
-	TUniquePtr<webrtc::AudioDeviceBuffer> DeviceBuffer;
-	FCriticalSection DeviceBufferCS; // protects capturing audio and sending it to WebRTC from simultaneous termination
-
-	TArray<uint8_t> RecordingBuffer;
-	int RecordingBufferSize = 0;
-
-	FThreadSafeBool bRecordingInitialized = false;
-	static constexpr int SampleRate = 48000;
-	static constexpr int NumChannels = 2;
-
-	bool bFormatChecked = false;
-
-	std::unique_ptr<webrtc::TaskQueueFactory> m_taskQueueFactory;
-};
+namespace UE {
+	namespace PixelStreaming {
+		// This class is only used if -PixelStreamingWebRTCUseLegacyAudioDevice or related CVar are used.
+		// This class will likely be removed in the future once the new ADM is confirmed stable.
+		class FAudioCapturer : public webrtc::AudioDeviceModule, public ISubmixBufferListener
+		{
+		private:
+			// ISubmixBufferListener interface
+			void OnNewSubmixBuffer(const USoundSubmix* OwningSubmix, float* AudioData, int32 NumSamples, int32 NumChannels, const int32 SampleRate, double AudioClock) override;
+
+			//
+			// webrtc::AudioDeviceModule interface
+			//
+			int32 ActiveAudioLayer(AudioLayer* audioLayer) const override;
+			int32 RegisterAudioCallback(webrtc::AudioTransport* audioCallback) override;
+
+			// Main initialization and termination
+			int32 Init() override;
+			int32 Terminate() override;
+			bool Initialized() const override;
+
+			// Device enumeration
+			int16 PlayoutDevices() override;
+			int16 RecordingDevices() override;
+			int32 PlayoutDeviceName(uint16 index, char name[webrtc::kAdmMaxDeviceNameSize], char guid[webrtc::kAdmMaxGuidSize]) override;
+			int32 RecordingDeviceName(uint16 index, char name[webrtc::kAdmMaxDeviceNameSize], char guid[webrtc::kAdmMaxGuidSize]) override;
+
+			// Device selection
+			int32 SetPlayoutDevice(uint16 index) override;
+			int32 SetPlayoutDevice(WindowsDeviceType device) override;
+			int32 SetRecordingDevice(uint16 index) override;
+			int32 SetRecordingDevice(WindowsDeviceType device) override;
+
+			// Audio transport initialization
+			int32 PlayoutIsAvailable(bool* available) override;
+			int32 InitPlayout() override;
+			bool PlayoutIsInitialized() const override;
+			int32 RecordingIsAvailable(bool* available) override;
+			int32 InitRecording() override;
+			bool RecordingIsInitialized() const override;
+
+			// Audio transport control
+			virtual int32 StartPlayout() override;
+			virtual int32 StopPlayout() override;
+			virtual bool Playing() const override;
+			virtual int32 StartRecording() override;
+			virtual int32 StopRecording() override;
+			virtual bool Recording() const override;
+
+			// Audio mixer initialization
+			virtual int32 InitSpeaker() override;
+			virtual bool SpeakerIsInitialized() const override;
+			virtual int32 InitMicrophone() override;
+			virtual bool MicrophoneIsInitialized() const override;
+
+			// Speaker volume controls
+			virtual int32 SpeakerVolumeIsAvailable(bool* available) override
+			{
+				return -1;
+			}
+			virtual int32 SetSpeakerVolume(uint32 volume) override
+			{
+				return -1;
+			}
+			virtual int32 SpeakerVolume(uint32* volume) const override
+			{
+				return -1;
+			}
+			virtual int32 MaxSpeakerVolume(uint32* maxVolume) const override
+			{
+				return -1;
+			}
+			virtual int32 MinSpeakerVolume(uint32* minVolume) const override
+			{
+				return -1;
+			}
+
+			// Microphone volume controls
+			virtual int32 MicrophoneVolumeIsAvailable(bool* available) override
+			{
+				return -1;
+			}
+			virtual int32 SetMicrophoneVolume(uint32 volume) override
+			{
+				return -1;
+			}
+			virtual int32 MicrophoneVolume(uint32* volume) const override
+			{
+				return -1;
+			}
+			virtual int32 MaxMicrophoneVolume(uint32* maxVolume) const override
+			{
+				return -1;
+			}
+			virtual int32 MinMicrophoneVolume(uint32* minVolume) const override
+			{
+				return -1;
+			}
+
+			// Speaker mute control
+			virtual int32 SpeakerMuteIsAvailable(bool* available) override
+			{
+				return -1;
+			}
+			virtual int32 SetSpeakerMute(bool enable) override
+			{
+				return -1;
+			}
+			virtual int32 SpeakerMute(bool* enabled) const override
+			{
+				return -1;
+			}
+
+			// Microphone mute control
+			virtual int32 MicrophoneMuteIsAvailable(bool* available) override
+			{
+				return -1;
+			}
+			virtual int32 SetMicrophoneMute(bool enable) override
+			{
+				return -1;
+			}
+			virtual int32 MicrophoneMute(bool* enabled) const override
+			{
+				return -1;
+			}
+
+			// Stereo support
+			virtual int32 StereoPlayoutIsAvailable(bool* available) const override;
+			virtual int32 SetStereoPlayout(bool enable) override;
+			virtual int32 StereoPlayout(bool* enabled) const override;
+			virtual int32 StereoRecordingIsAvailable(bool* available) const override;
+			virtual int32 SetStereoRecording(bool enable) override;
+			virtual int32 StereoRecording(bool* enabled) const override;
+
+			// Playout delay
+			virtual int32 PlayoutDelay(uint16* delayMS) const override
+			{
+				return -1;
+			}
+
+			// Only supported on Android.
+			virtual bool BuiltInAECIsAvailable() const override
+			{
+				return false;
+			}
+			virtual bool BuiltInAGCIsAvailable() const override
+			{
+				return false;
+			}
+			virtual bool BuiltInNSIsAvailable() const override
+			{
+				return false;
+			}
+
+			// Enables the built-in audio effects. Only supported on Android.
+			virtual int32 EnableBuiltInAEC(bool enable) override
+			{
+				return -1;
+			}
+			virtual int32 EnableBuiltInAGC(bool enable) override
+			{
+				return -1;
+			}
+			virtual int32 EnableBuiltInNS(bool enable) override
+			{
+				return -1;
+			}
+
+			FThreadSafeBool bInitialized = false;
+			TArray<int16> PCM16;
+
+			TUniquePtr<webrtc::AudioDeviceBuffer> DeviceBuffer;
+			FCriticalSection DeviceBufferCS; // protects capturing audio and sending it to WebRTC from simultaneous termination
+
+			TArray<uint8_t> RecordingBuffer;
+			int RecordingBufferSize = 0;
+
+			FThreadSafeBool bRecordingInitialized = false;
+			static constexpr int SampleRate = 48000;
+			static constexpr int NumChannels = 2;
+
+			bool bFormatChecked = false;
+
+			std::unique_ptr<webrtc::TaskQueueFactory> m_taskQueueFactory;
+		};
+	}
+}
diff --git a/Engine/Plugins/Media/PixelStreaming/Source/PixelStreaming/Private/AudioDeviceModule.cpp b/Engine/Plugins/Media/PixelStreaming/Source/PixelStreaming/Private/AudioDeviceModule.cpp
new file mode 100644
--- /dev/null
+++ b/Engine/Plugins/Media/PixelStreaming/Source/PixelStreaming/Private/AudioDeviceModule.cpp
@@ -0,0 +1,613 @@
+// Copyright Epic Games, Inc. All Rights Reserved.
+
+#include "AudioDeviceModule.h"
+#include "AudioMixerDevice.h"
+#include "SampleBuffer.h"
+#include "Engine/GameEngine.h"
+#include "Settings.h"
+
+DECLARE_LOG_CATEGORY_EXTERN(LogPixelStreamingAudioDeviceModule, Log, All);
+DEFINE_LOG_CATEGORY(LogPixelStreamingAudioDeviceModule);
+
+// These are copied from webrtc internals
+#define CHECKinitialized_() \
+	{                       \
+		if (!bInitialized)  \
+		{                   \
+			return -1;      \
+		};                  \
+	}
+#define CHECKinitialized__BOOL() \
+	{                            \
+		if (!bInitialized)       \
+		{                        \
+			return false;        \
+		};                       \
+	}
+
+int32 UE::PixelStreaming::FAudioDeviceModule::ActiveAudioLayer(AudioLayer* audioLayer) const
+{
+	*audioLayer = AudioDeviceModule::kDummyAudio;
+	return 0;
+}
+
+int32 UE::PixelStreaming::FAudioDeviceModule::RegisterAudioCallback(webrtc::AudioTransport* audioCallback)
+{
+	Requester.RegisterAudioCallback(audioCallback);
+	Capturer.RegisterAudioCallback(audioCallback);
+	return 0;
+}
+
+int32 UE::PixelStreaming::FAudioDeviceModule::Init()
+{
+	InitRecording();
+
+	UE_LOG(LogPixelStreamingAudioDeviceModule, Verbose, TEXT("Init PixelStreamingAudioDeviceModule"));
+	bInitialized = true;
+	return 0;
+}
+
+int32 UE::PixelStreaming::FAudioDeviceModule::Terminate()
+{
+	if (!bInitialized)
+	{
+		return 0;
+	}
+
+	Requester.Uninitialise();
+	Capturer.Uninitialise();
+
+	UE_LOG(LogPixelStreamingAudioDeviceModule, Verbose, TEXT("Terminate"));
+
+	return 0;
+}
+
+bool UE::PixelStreaming::FAudioDeviceModule::Initialized() const
+{
+	return bInitialized;
+}
+
+int16 UE::PixelStreaming::FAudioDeviceModule::PlayoutDevices()
+{
+	CHECKinitialized_();
+	return -1;
+}
+
+int16 UE::PixelStreaming::FAudioDeviceModule::RecordingDevices()
+{
+	CHECKinitialized_();
+	return -1;
+}
+
+int32 UE::PixelStreaming::FAudioDeviceModule::PlayoutDeviceName(uint16 index, char name[webrtc::kAdmMaxDeviceNameSize], char guid[webrtc::kAdmMaxGuidSize])
+{
+	CHECKinitialized_();
+	return -1;
+}
+
+int32 UE::PixelStreaming::FAudioDeviceModule::RecordingDeviceName(uint16 index, char name[webrtc::kAdmMaxDeviceNameSize], char guid[webrtc::kAdmMaxGuidSize])
+{
+	CHECKinitialized_();
+	return -1;
+}
+
+int32 UE::PixelStreaming::FAudioDeviceModule::SetPlayoutDevice(uint16 index)
+{
+	CHECKinitialized_();
+	return 0;
+}
+
+int32 UE::PixelStreaming::FAudioDeviceModule::SetPlayoutDevice(WindowsDeviceType device)
+{
+	CHECKinitialized_();
+	return 0;
+}
+
+int32 UE::PixelStreaming::FAudioDeviceModule::SetRecordingDevice(uint16 index)
+{
+	CHECKinitialized_();
+	return 0;
+}
+
+int32 UE::PixelStreaming::FAudioDeviceModule::SetRecordingDevice(WindowsDeviceType device)
+{
+	CHECKinitialized_();
+	return 0;
+}
+
+int32 UE::PixelStreaming::FAudioDeviceModule::PlayoutIsAvailable(bool* available)
+{
+	CHECKinitialized_();
+	*available = !UE::PixelStreaming::Settings::CVarPixelStreamingWebRTCDisableReceiveAudio.GetValueOnAnyThread();
+	return 0;
+}
+
+int32 UE::PixelStreaming::FAudioDeviceModule::InitPlayout()
+{
+	CHECKinitialized_();
+	bool bIsPlayoutAvailable = false;
+	PlayoutIsAvailable(&bIsPlayoutAvailable);
+	if (!bIsPlayoutAvailable)
+	{
+		return -1;
+	}
+
+	Requester.InitPlayout();
+	return 0;
+}
+
+bool UE::PixelStreaming::FAudioDeviceModule::PlayoutIsInitialized() const
+{
+	CHECKinitialized__BOOL();
+	return Requester.PlayoutIsInitialized();
+}
+
+int32 UE::PixelStreaming::FAudioDeviceModule::StartPlayout()
+{
+	CHECKinitialized_();
+
+	if (!Requester.PlayoutIsInitialized())
+	{
+		return -1;
+	}
+
+	Requester.StartPlayout();
+	return 0;
+}
+
+int32 UE::PixelStreaming::FAudioDeviceModule::StopPlayout()
+{
+	CHECKinitialized_();
+	if (!Requester.PlayoutIsInitialized())
+	{
+		return -1;
+	}
+
+	Requester.StopPlayout();
+	return 0;
+}
+
+bool UE::PixelStreaming::FAudioDeviceModule::Playing() const
+{
+	CHECKinitialized__BOOL();
+	return Requester.Playing();
+}
+
+int32 UE::PixelStreaming::FAudioDeviceModule::RecordingIsAvailable(bool* available)
+{
+	CHECKinitialized_();
+	*available = !UE::PixelStreaming::Settings::CVarPixelStreamingWebRTCDisableTransmitAudio.GetValueOnAnyThread();
+	return 0;
+}
+
+int32 UE::PixelStreaming::FAudioDeviceModule::InitRecording()
+{
+	CHECKinitialized_();
+
+	bool bIsRecordingAvailable = false;
+	RecordingIsAvailable(&bIsRecordingAvailable);
+	if (!bIsRecordingAvailable)
+	{
+		return -1;
+	}
+
+	if (!Capturer.IsInitialised())
+	{
+		Capturer.Init();
+	}
+
+	return 0;
+}
+
+bool UE::PixelStreaming::FAudioDeviceModule::RecordingIsInitialized() const
+{
+	CHECKinitialized__BOOL();
+	return Capturer.IsInitialised();
+}
+
+int32 UE::PixelStreaming::FAudioDeviceModule::StartRecording()
+{
+	CHECKinitialized_();
+	if (!Capturer.IsCapturing())
+	{
+		Capturer.StartCapturing();
+	}
+	return 0;
+}
+
+int32 UE::PixelStreaming::FAudioDeviceModule::StopRecording()
+{
+	CHECKinitialized_();
+	if (Capturer.IsCapturing())
+	{
+		Capturer.EndCapturing();
+	}
+	return 0;
+}
+
+bool UE::PixelStreaming::FAudioDeviceModule::Recording() const
+{
+	CHECKinitialized__BOOL();
+	return Capturer.IsCapturing();
+}
+
+int32 UE::PixelStreaming::FAudioDeviceModule::InitSpeaker()
+{
+	CHECKinitialized_();
+	return -1;
+}
+
+bool UE::PixelStreaming::FAudioDeviceModule::SpeakerIsInitialized() const
+{
+	CHECKinitialized__BOOL();
+	return false;
+}
+
+int32 UE::PixelStreaming::FAudioDeviceModule::InitMicrophone()
+{
+	CHECKinitialized_();
+	return 0;
+}
+
+bool UE::PixelStreaming::FAudioDeviceModule::MicrophoneIsInitialized() const
+{
+	CHECKinitialized__BOOL();
+	return true;
+}
+
+int32 UE::PixelStreaming::FAudioDeviceModule::StereoPlayoutIsAvailable(bool* available) const
+{
+	CHECKinitialized_();
+	*available = true;
+	return 0;
+}
+
+int32 UE::PixelStreaming::FAudioDeviceModule::SetStereoPlayout(bool enable)
+{
+	CHECKinitialized_();
+	FString AudioChannelStr = enable ? TEXT("stereo") : TEXT("mono");
+	UE_LOG(LogPixelStreamingAudioDeviceModule, Verbose, TEXT("WebRTC has requested browser audio playout in UE be: %s"), *AudioChannelStr);
+	return 0;
+}
+
+int32 UE::PixelStreaming::FAudioDeviceModule::StereoPlayout(bool* enabled) const
+{
+	CHECKinitialized_();
+	*enabled = true;
+	return 0;
+}
+
+int32 UE::PixelStreaming::FAudioDeviceModule::StereoRecordingIsAvailable(bool* available) const
+{
+	CHECKinitialized_();
+	*available = true;
+	return 0;
+}
+
+int32 UE::PixelStreaming::FAudioDeviceModule::SetStereoRecording(bool enable)
+{
+	CHECKinitialized_();
+	return 0;
+}
+
+int32 UE::PixelStreaming::FAudioDeviceModule::StereoRecording(bool* enabled) const
+{
+	CHECKinitialized_();
+	*enabled = true;
+	return 0;
+}
+
+//--UE::PixelStreaming::FSubmixCapturer--
+
+uint32_t UE::PixelStreaming::FSubmixCapturer::GetVolume() const
+{
+	return VolumeLevel;
+}
+
+void UE::PixelStreaming::FSubmixCapturer::SetVolume(uint32_t NewVolume)
+{
+	VolumeLevel = NewVolume;
+}
+
+bool UE::PixelStreaming::FSubmixCapturer::Init()
+{
+	FScopeLock Lock(&CriticalSection);
+
+	// subscribe to audio data
+	if (!GEngine)
+	{
+		bInitialised = false;
+		return false;
+	}
+
+	// already initialised
+	if (bInitialised)
+	{
+		return true;
+	}
+
+	FAudioDeviceHandle AudioDevice = GEngine->GetMainAudioDevice();
+	if (!AudioDevice)
+	{
+		UE_LOG(LogPixelStreamingAudioDeviceModule, Warning, TEXT("No audio device"));
+		bInitialised = false;
+		return false;
+	}
+
+	AudioDevice->RegisterSubmixBufferListener(this);
+	bInitialised = true;
+	return true;
+}
+
+void UE::PixelStreaming::FSubmixCapturer::OnNewSubmixBuffer(const USoundSubmix* OwningSubmix, float* AudioData, int32 NumSamples, int32 NumChannels, const int32 SampleRate, double AudioClock)
+{
+	FScopeLock Lock(&CriticalSection);
+
+	if (!bInitialised || !bCapturing)
+	{
+		return;
+	}
+
+	// No point doing anything with UE audio if the callback from WebRTC has not been set yet.
+	if (AudioCallback == nullptr)
+	{
+		return;
+	}
+
+	// Check if the sample rate from UE matches our target sample rate
+	if (TargetSampleRate != SampleRate)
+	{
+		// Only report the problem once
+		if (!bReportedSampleRateMismatch)
+		{
+			bReportedSampleRateMismatch = true;
+			UE_LOG(LogPixelStreamingAudioDeviceModule, Error, TEXT("Audio sample rate mismatch. Expected: %d | Actual: %d"), TargetSampleRate, SampleRate);
+		}
+		return;
+	}
+
+	UE_LOG(LogPixelStreamingAudioDeviceModule, VeryVerbose, TEXT("captured %d samples, %dc, %dHz"), NumSamples, NumChannels, SampleRate);
+
+	// Note: TSampleBuffer takes in AudioData as float* and internally converts to int16
+	Audio::TSampleBuffer<int16> Buffer(AudioData, NumSamples, NumChannels, SampleRate);
+
+	// Mix to our target number of channels if the source does not already match.
+	if (Buffer.GetNumChannels() != TargetNumChannels)
+	{
+		Buffer.MixBufferToChannels(TargetNumChannels);
+	}
+
+	RecordingBuffer.Append(Buffer.GetData(), Buffer.GetNumSamples());
+
+	const float ChunkDurationSecs = 0.01f; //10ms
+	const int32 SamplesPer10Ms = GetSamplesPerDurationSecs(ChunkDurationSecs);
+
+	// Feed in 10ms chunks
+	while (RecordingBuffer.Num() > SamplesPer10Ms)
+	{
+
+		// Extract a 10ms chunk of samples from recording buffer
+		TArray<int16_t> SubmitBuffer(RecordingBuffer.GetData(), SamplesPer10Ms);
+		const size_t Frames = SubmitBuffer.Num() / TargetNumChannels;
+		const size_t BytesPerFrame = TargetNumChannels * sizeof(int16_t);
+
+		uint32_t OutMicLevel = VolumeLevel;
+
+		int32_t WebRTCRes = AudioCallback->RecordedDataIsAvailable(
+			SubmitBuffer.GetData(),
+			Frames,
+			BytesPerFrame,
+			TargetNumChannels,
+			TargetSampleRate,
+			0,
+			0,
+			VolumeLevel,
+			false,
+			OutMicLevel);
+
+		SetVolume(OutMicLevel);
+
+		// Remove 10ms of samples from the recording buffer now it is submitted
+		RecordingBuffer.RemoveAt(0, SamplesPer10Ms, false);
+	}
+}
+
+int32 UE::PixelStreaming::FSubmixCapturer::GetSamplesPerDurationSecs(float InSeconds) const
+{
+	int32 SamplesPerSecond = TargetNumChannels * TargetSampleRate;
+	int32 NumSamplesPerDuration = (int32)(SamplesPerSecond * InSeconds);
+	return NumSamplesPerDuration;
+}
+
+void UE::PixelStreaming::FSubmixCapturer::RegisterAudioCallback(webrtc::AudioTransport* AudioTransportCallback)
+{
+	FScopeLock Lock(&CriticalSection);
+	AudioCallback = AudioTransportCallback;
+}
+
+bool UE::PixelStreaming::FSubmixCapturer::IsInitialised() const
+{
+	return bInitialised;
+}
+
+void UE::PixelStreaming::FSubmixCapturer::Uninitialise()
+{
+	FScopeLock Lock(&CriticalSection);
+
+	if (GEngine)
+	{
+		FAudioDeviceHandle AudioDevice = GEngine->GetMainAudioDevice();
+		if (AudioDevice)
+		{
+			AudioDevice->UnregisterSubmixBufferListener(this);
+		}
+	}
+
+	RecordingBuffer.Empty();
+	bInitialised = false;
+	bCapturing = false;
+}
+
+bool UE::PixelStreaming::FSubmixCapturer::StartCapturing()
+{
+	if (!bInitialised)
+	{
+		return false;
+	}
+	bCapturing = true;
+	return true;
+}
+
+bool UE::PixelStreaming::FSubmixCapturer::EndCapturing()
+{
+	if (!bInitialised)
+	{
+		return false;
+	}
+	bCapturing = false;
+	return true;
+}
+
+bool UE::PixelStreaming::FSubmixCapturer::IsCapturing() const
+{
+	return bCapturing;
+}
+
+//--UE::PixelStreaming::FAudioPlayoutRequester--
+
+void UE::PixelStreaming::FAudioPlayoutRequester::Uninitialise()
+{
+	StopPlayout();
+	AudioCallback = nullptr;
+	bIsPlayoutInitialised = false;
+	PlayoutBuffer.Empty();
+}
+
+void UE::PixelStreaming::FAudioPlayoutRequester::InitPlayout()
+{
+	bIsPlayoutInitialised = true;
+}
+
+void UE::PixelStreaming::FAudioPlayoutRequester::StartPlayout()
+{
+	if (PlayoutIsInitialized() && !Playing())
+	{
+		TFunction<void()> RequesterFunc = [this]() {
+			FScopeLock Lock(&PlayoutCriticalSection);
+
+			// Only request audio if audio callback is valid
+			if (!AudioCallback)
+			{
+				return;
+			}
+
+			// Our intention is to request samples at some fixed interval (i.e 10ms)
+			int32 NSamplesPerChannel = (SampleRate) / (1000 / UE::PixelStreaming::FAudioPlayoutRequester::RequestIntervalMs);
+			const size_t BytesPerFrame = NumChannels * sizeof(int16_t);
+
+			// Ensure buffer has the total number of samples we need for all audio frames
+			PlayoutBuffer.Reserve(NSamplesPerChannel * NumChannels);
+
+			size_t OutNSamples = 0;
+			int64_t ElapsedTimeMs = -1;
+			int64_t NtpTimeMs = -1;
+
+			// Note this is mixed result of all audio sources, which in turn triggers the sinks for each audio source to be called.
+			// For example, if you had 3 audio sources in the browser sending 16kHz mono they would all be mixed down into the number
+			// of channels and sample rate specified below. However, for listening to each audio source prior to mixing refer to
+			// UE::PixelStreaming::FAudioSink.
+			uint32_t Result = AudioCallback->NeedMorePlayData(
+				NSamplesPerChannel,
+				BytesPerFrame,
+				NumChannels,
+				SampleRate,
+				PlayoutBuffer.GetData(),
+				OutNSamples,
+				&ElapsedTimeMs,
+				&NtpTimeMs);
+
+			if (Result != 0)
+			{
+				UE_LOG(LogPixelStreamingAudioDeviceModule, Error, TEXT("NeedMorePlayData return non-zero result indicating an error"));
+			}
+		};
+
+		RequesterThread.Reset(nullptr);
+		RequesterRunnable = MakeUnique<UE::PixelStreaming::FAudioPlayoutRequester::Runnable>(RequesterFunc);
+		RequesterThread.Reset(FRunnableThread::Create(RequesterRunnable.Get(), TEXT("Pixel Streaming WebRTC Audio Requester")));
+		bIsPlaying = true;
+	}
+}
+
+void UE::PixelStreaming::FAudioPlayoutRequester::StopPlayout()
+{
+	if (PlayoutIsInitialized() && Playing())
+	{
+		RequesterRunnable->Stop();
+		bIsPlaying = false;
+	}
+}
+
+bool UE::PixelStreaming::FAudioPlayoutRequester::Playing() const
+{
+	return bIsPlaying;
+}
+
+bool UE::PixelStreaming::FAudioPlayoutRequester::PlayoutIsInitialized() const
+{
+	return bIsPlayoutInitialised;
+}
+
+void UE::PixelStreaming::FAudioPlayoutRequester::RegisterAudioCallback(webrtc::AudioTransport* AudioCb)
+{
+	FScopeLock Lock(&PlayoutCriticalSection);
+	AudioCallback = AudioCb;
+}
+
+//--FAudioRequester::Runnable--
+
+bool UE::PixelStreaming::FAudioPlayoutRequester::Runnable::Init()
+{
+	return RequestPlayoutFunc != nullptr;
+}
+
+uint32 UE::PixelStreaming::FAudioPlayoutRequester::Runnable::Run()
+{
+	LastAudioRequestTimeMs = rtc::TimeMillis();
+	bIsRunning = true;
+
+	// Request audio in a loop until this boolean is toggled off.
+	while (bIsRunning)
+	{
+		int64_t Now = rtc::TimeMillis();
+		int64_t DeltaMs = Now - LastAudioRequestTimeMs;
+
+		// Check if the 10ms delta has elapsed, if it has not, then sleep the remaining
+		if (DeltaMs < UE::PixelStreaming::FAudioPlayoutRequester::RequestIntervalMs)
+		{
+			int64_t SleepTimeMs = UE::PixelStreaming::FAudioPlayoutRequester::RequestIntervalMs - DeltaMs;
+			float SleepTimeSecs = (float)SleepTimeMs / 1000.0f;
+			FPlatformProcess::Sleep(SleepTimeSecs);
+			continue;
+		}
+
+		// Update request time to now seeing as the 10ms delta has elapsed
+		LastAudioRequestTimeMs = Now;
+
+		// Actually request playout
+		RequestPlayoutFunc();
+	}
+
+	return 0;
+}
+
+void UE::PixelStreaming::FAudioPlayoutRequester::Runnable::Stop()
+{
+	bIsRunning = false;
+}
+
+void UE::PixelStreaming::FAudioPlayoutRequester::Runnable::Exit()
+{
+	bIsRunning = false;
+	RequestPlayoutFunc = nullptr;
+}
\ No newline at end of file
