Temp early exit when using AMF on DX11.

diff --git a/Engine/Plugins/Media/HardwareEncoders/Source/EncoderAMF/Private/Amf_EncoderH264.cpp b/Engine/Plugins/Media/HardwareEncoders/Source/EncoderAMF/Private/Amf_EncoderH264.cpp
--- a/Engine/Plugins/Media/HardwareEncoders/Source/EncoderAMF/Private/Amf_EncoderH264.cpp
+++ b/Engine/Plugins/Media/HardwareEncoders/Source/EncoderAMF/Private/Amf_EncoderH264.cpp
@@ -169,6 +169,13 @@ namespace AVEncoder
 			return false;
 		}
 
+		//TODO(sandor.hadas) see if the current issues with AMF can be resolved then remove this error message
+		if (RHIName == "D3D11")
+		{
+			UE_LOG(LogEncoderAMF, Error, TEXT("AMF with DX11 is not currently supported try DX12 or Vulkan."));
+			return false;
+		}
+
 		if (!EncoderDevice)
 		{
 			UE_LOG(LogEncoderAMF, Error, TEXT("Amf needs an encoder device."));
