- Fixed issue with incorrect hardware encoders loading when the target GPU is not present.
- Vulkan AMF support is currently unsupported and has been disabled.
- Removed fatal error when trying to use AMF with Vulkan.

diff --git a/Engine/Plugins/Media/HardwareEncoders/Source/EncoderAMF/Private/Amf_EncoderModule.cpp b/Engine/Plugins/Media/HardwareEncoders/Source/EncoderAMF/Private/Amf_EncoderModule.cpp
--- a/Engine/Plugins/Media/HardwareEncoders/Source/EncoderAMF/Private/Amf_EncoderModule.cpp
+++ b/Engine/Plugins/Media/HardwareEncoders/Source/EncoderAMF/Private/Amf_EncoderModule.cpp
@@ -32,8 +32,9 @@ public:
 #endif
 				if(FString("VulkanRHI") == FString(DynamicRHIModuleName))
 				{
-#if PLATFORM_WINDOWS
-					UE_LOG(LogEncoderAMF, Fatal, TEXT("Vulkan Amf support under Windows is currently unstable and has been hard disabled for the time being."));
+#if PLATFORM_WINDOWS || PLATFORM_LINUX
+					UE_LOG(LogEncoderAMF, Error, TEXT("Vulkan AMF is currently unsuported and has been disabled."));
+					return;
 #endif
 					AMF.InitializeContext("Vulkan", NULL);
 					amf::AMFContext1Ptr pContext1(AMF.GetContext());
