NoPCH Compile LyraEditor Win64 - Amf_Common.h error: syntax error: identifier 'ERHIInterfaceType'

diff --git a/Engine/Plugins/Media/HardwareEncoders/Source/EncoderAMF/Private/Amf_Common.h b/Engine/Plugins/Media/HardwareEncoders/Source/EncoderAMF/Private/Amf_Common.h
--- a/Engine/Plugins/Media/HardwareEncoders/Source/EncoderAMF/Private/Amf_Common.h
+++ b/Engine/Plugins/Media/HardwareEncoders/Source/EncoderAMF/Private/Amf_Common.h
@@ -3,6 +3,7 @@
 #pragma once
 
 #include "CoreMinimal.h"
+#include "RHIDefinitions.h"
 #include "HAL/Thread.h"
 
 #include <string.h>
