Fix to fail safely when AMF encoder fails to initialize on some AMD GPU + Driver combinations.

diff --git a/Engine/Plugins/Media/HardwareEncoders/Source/EncoderAMF/Private/Amf_EncoderH264.cpp b/Engine/Plugins/Media/HardwareEncoders/Source/EncoderAMF/Private/Amf_EncoderH264.cpp
--- a/Engine/Plugins/Media/HardwareEncoders/Source/EncoderAMF/Private/Amf_EncoderH264.cpp
+++ b/Engine/Plugins/Media/HardwareEncoders/Source/EncoderAMF/Private/Amf_EncoderH264.cpp
@@ -283,7 +283,7 @@ namespace AVEncoder
 	{
 		if (AmfEncoder == NULL)
 		{
-			Amf.CreateEncoder(AmfEncoder);
+			return Amf.CreateEncoder(AmfEncoder) && AmfEncoder != NULL;
 		}
 
 		return AmfEncoder != NULL;
@@ -798,7 +798,13 @@ namespace AVEncoder
 		
 		// Create temp component
 		AMFComponentPtr TempEncoder;
-		AMF.CreateEncoder(TempEncoder);
+		bool bCreatedEncoder = AMF.CreateEncoder(TempEncoder);
+
+		if(!bCreatedEncoder || TempEncoder == NULL)
+		{
+			UE_LOG(LogEncoderAMF, Warning, TEXT("Failed to created AMF encoder on AMD hardware. Consider trying a different driver version."));
+			return false;
+		}
 
 		AMFCapsPtr EncoderCaps;
 		TempEncoder->GetCaps(&EncoderCaps);
